UE4: Text change to make reference viewer "Map" references instead say "World" to be consistent for when WorldAssets are enabled

diff --git a/Engine/Source/Editor/ReferenceViewer/Private/EdGraphNode_Reference.cpp b/Engine/Source/Editor/ReferenceViewer/Private/EdGraphNode_Reference.cpp
--- a/Engine/Source/Editor/ReferenceViewer/Private/EdGraphNode_Reference.cpp
+++ b/Engine/Source/Editor/ReferenceViewer/Private/EdGraphNode_Reference.cpp
@@ -141,7 +141,7 @@ void UEdGraphNode_Reference::CacheAssetData(const FAssetData& AssetData)
 					const bool bIsMapPackage = FPlatformFileManager::Get().GetPlatformFile().FileExists(*PotentiallyMapFilename);
 					if ( bIsMapPackage )
 					{
-						CachedAssetData.AssetClass = FName(TEXT("Map"));
+						CachedAssetData.AssetClass = FName(TEXT("World"));
 					}
 				}
 			}
