Add Edit Asset option to the Reference Viewer context menu.
#ttp 333623 Allow users to open asset editors straight from the Reference Viewer
reviewer Thomas.Sarkanen

diff --git a/Engine/Source/Editor/ReferenceViewer/Private/ReferenceViewerActions.cpp b/Engine/Source/Editor/ReferenceViewer/Private/ReferenceViewerActions.cpp
--- a/Engine/Source/Editor/ReferenceViewer/Private/ReferenceViewerActions.cpp
+++ b/Engine/Source/Editor/ReferenceViewer/Private/ReferenceViewerActions.cpp
@@ -8,6 +8,7 @@
 
 void FReferenceViewerActions::RegisterCommands()
 {
+	UI_COMMAND(OpenSelectedInAssetEditor, "Edit Asset", "Opens the selected asset in the relevent editor.", EUserInterfaceActionType::Button, FInputGesture());
 	UI_COMMAND(ReCenterGraph, "Re-Center Graph", "Re-centers the graph on this node, showing all referencers and dependencies for this asset instead", EUserInterfaceActionType::Button, FInputGesture());
 	UI_COMMAND(ListReferencedObjects, "List Referenced Objects", "Shows a list of objects that the selected asset references.", EUserInterfaceActionType::Button, FInputGesture());
 	UI_COMMAND(ListObjectsThatReference, "List Objects That Reference", "Lists objects that reference the selected asset.", EUserInterfaceActionType::Button, FInputGesture());
diff --git a/Engine/Source/Editor/ReferenceViewer/Private/ReferenceViewerActions.h b/Engine/Source/Editor/ReferenceViewer/Private/ReferenceViewerActions.h
--- a/Engine/Source/Editor/ReferenceViewer/Private/ReferenceViewerActions.h
+++ b/Engine/Source/Editor/ReferenceViewer/Private/ReferenceViewerActions.h
@@ -20,6 +20,9 @@ public:
 	// End of TCommands<> interface
 public:
 
+	// Opens the selected asset in the asset editor
+	TSharedPtr<FUICommandInfo> OpenSelectedInAssetEditor;
+
 	// Re-constructs the graph with the selected asset as the center
 	TSharedPtr<FUICommandInfo> ReCenterGraph;
 
diff --git a/Engine/Source/Editor/ReferenceViewer/Private/ReferenceViewerSchema.cpp b/Engine/Source/Editor/ReferenceViewer/Private/ReferenceViewerSchema.cpp
--- a/Engine/Source/Editor/ReferenceViewer/Private/ReferenceViewerSchema.cpp
+++ b/Engine/Source/Editor/ReferenceViewer/Private/ReferenceViewerSchema.cpp
@@ -16,6 +16,7 @@ UReferenceViewerSchema::UReferenceViewerSchema(const class FPostConstructInitial
 void UReferenceViewerSchema::GetContextMenuActions(const UEdGraph* CurrentGraph, const UEdGraphNode* InGraphNode, const UEdGraphPin* InGraphPin, class FMenuBuilder* MenuBuilder, bool bIsDebugging) const
 {
 	MenuBuilder->AddMenuEntry(FGlobalEditorCommonCommands::Get().FindInContentBrowser);
+	MenuBuilder->AddMenuEntry(FReferenceViewerActions::Get().OpenSelectedInAssetEditor);
 	MenuBuilder->AddMenuEntry(FReferenceViewerActions::Get().ReCenterGraph);
 	MenuBuilder->AddMenuEntry(FReferenceViewerActions::Get().ListReferencedObjects);
 	MenuBuilder->AddMenuEntry(FReferenceViewerActions::Get().ListObjectsThatReference);
diff --git a/Engine/Source/Editor/ReferenceViewer/Private/SReferenceViewer.cpp b/Engine/Source/Editor/ReferenceViewer/Private/SReferenceViewer.cpp
--- a/Engine/Source/Editor/ReferenceViewer/Private/SReferenceViewer.cpp
+++ b/Engine/Source/Editor/ReferenceViewer/Private/SReferenceViewer.cpp
@@ -3,6 +3,7 @@
 #include "ReferenceViewerPrivatePCH.h"
 #include "GraphEditor.h"
 #include "AssetRegistryModule.h"
+#include "AssetEditorManager.h"
 #include "AssetThumbnail.h"
 #include "ReferenceViewerActions.h"
 #include "EditorWidgets.h"
@@ -477,6 +478,13 @@ void SReferenceViewer::RegisterActions()
 		FGlobalEditorCommonCommands::Get().FindInContentBrowser,
 		FExecuteAction::CreateSP(this, &SReferenceViewer::ShowSelectionInContentBrowser));
 
+	ReferenceViewerActions->MapAction(
+		FReferenceViewerActions::Get().OpenSelectedInAssetEditor,
+		FExecuteAction::CreateSP(this, &SReferenceViewer::OpenSelectedInAssetEditor),
+		FCanExecuteAction::CreateSP(this, &SReferenceViewer::HasExactlyOneNodeSelected),
+		FIsActionChecked(),
+		FIsActionButtonVisible::CreateSP(this, &SReferenceViewer::IsSingleSelectedItemValidObject));
+
 	ReferenceViewerActions->MapAction(
 		FReferenceViewerActions::Get().ReCenterGraph,
 		FExecuteAction::CreateSP(this, &SReferenceViewer::ReCenterGraph),
@@ -533,6 +541,16 @@ void SReferenceViewer::ShowSelectionInContentBrowser()
 	}
 }
 
+void SReferenceViewer::OpenSelectedInAssetEditor()
+{
+	UObject* SelectedObject = GetObjectFromSingleSelectedNode();
+
+	if( SelectedObject )
+	{
+		FAssetEditorManager::Get().OpenEditorForAsset(SelectedObject);
+	}
+}
+
 void SReferenceViewer::ReCenterGraph()
 {
 	ReCenterGraphOnNodes( GraphEditorPtr->GetSelectedNodes() );
@@ -660,6 +678,17 @@ bool SReferenceViewer::HasExactlyOneNodeSelected() const
 	return false;
 }
 
+bool SReferenceViewer::IsSingleSelectedItemValidObject() const
+{
+	UObject* SelectedObject = GetObjectFromSingleSelectedNode();
+
+	if ( SelectedObject )
+	{
+		return true;
+	}
+	return false;
+}
+
 void SReferenceViewer::OnInitialAssetRegistrySearchComplete()
 {
 	if ( GraphObj )
diff --git a/Engine/Source/Editor/ReferenceViewer/Private/SReferenceViewer.h b/Engine/Source/Editor/ReferenceViewer/Private/SReferenceViewer.h
--- a/Engine/Source/Editor/ReferenceViewer/Private/SReferenceViewer.h
+++ b/Engine/Source/Editor/ReferenceViewer/Private/SReferenceViewer.h
@@ -73,6 +73,7 @@ private:
 
 	void RegisterActions();
 	void ShowSelectionInContentBrowser();
+	void OpenSelectedInAssetEditor();
 	void ReCenterGraph();
 	void ListReferencedObjects();
 	void ListObjectsThatReference();
@@ -83,6 +84,7 @@ private:
 
 	UObject* GetObjectFromSingleSelectedNode() const;
 	bool HasExactlyOneNodeSelected() const;
+	bool IsSingleSelectedItemValidObject() const;
 
 	void OnInitialAssetRegistrySearchComplete();
 private:
