Removed the invalid right-click context menu from the Reference Viewer when clicking on the background.

diff --git a/Engine/Source/Editor/ReferenceViewer/Private/SReferenceViewer.cpp b/Engine/Source/Editor/ReferenceViewer/Private/SReferenceViewer.cpp
--- a/Engine/Source/Editor/ReferenceViewer/Private/SReferenceViewer.cpp
+++ b/Engine/Source/Editor/ReferenceViewer/Private/SReferenceViewer.cpp
@@ -40,6 +40,7 @@ void SReferenceViewer::Construct( const FArguments& InArgs )
 
 	SGraphEditor::FGraphEditorEvents GraphEvents;
 	GraphEvents.OnNodeDoubleClicked = FSingleNodeEvent::CreateSP(this, &SReferenceViewer::OnNodeDoubleClicked);
+	GraphEvents.OnCreateActionMenu = SGraphEditor::FOnCreateActionMenu::CreateSP(this, &SReferenceViewer::OnCreateGraphActionMenu);
 
 	// Create the graph editor
 	GraphEditorPtr = SNew(SGraphEditor)
@@ -260,6 +261,12 @@ void SReferenceViewer::OnNodeDoubleClicked(UEdGraphNode* Node)
 	ReCenterGraphOnNodes( Nodes );
 }
 
+FActionMenuContent SReferenceViewer::OnCreateGraphActionMenu(UEdGraph* InGraph, const FVector2D& InNodePosition, const TArray<UEdGraphPin*>& InDraggedPins, bool bAutoExpand, SGraphEditor::FActionMenuClosed InOnMenuClosed)
+{
+	// no context menu when not over a node
+	return FActionMenuContent();
+}
+
 bool SReferenceViewer::IsBackEnabled() const
 {
 	return HistoryManager.CanGoBack();
diff --git a/Engine/Source/Editor/ReferenceViewer/Private/SReferenceViewer.h b/Engine/Source/Editor/ReferenceViewer/Private/SReferenceViewer.h
--- a/Engine/Source/Editor/ReferenceViewer/Private/SReferenceViewer.h
+++ b/Engine/Source/Editor/ReferenceViewer/Private/SReferenceViewer.h
@@ -24,6 +24,10 @@ public:
 
 private:
 
+	/** Called to create context menu when right-clicking on graph */
+	FActionMenuContent OnCreateGraphActionMenu(UEdGraph* InGraph, const FVector2D& InNodePosition, const TArray<UEdGraphPin*>& InDraggedPins, bool bAutoExpand, SGraphEditor::FActionMenuClosed InOnMenuClosed);
+
+	/** Called when a node is double clicked */
 	void OnNodeDoubleClicked(class UEdGraphNode* Node);
 
 	/** True if the user may use the history back button */
