Size Mapper fixes
- Fixed Size Map displaying "unknown" when viewing a size map for a selected class in Reference Viewer (UE-14488)

diff --git a/Engine/Source/Editor/ReferenceViewer/Private/SReferenceViewer.cpp b/Engine/Source/Editor/ReferenceViewer/Private/SReferenceViewer.cpp
--- a/Engine/Source/Editor/ReferenceViewer/Private/SReferenceViewer.cpp
+++ b/Engine/Source/Editor/ReferenceViewer/Private/SReferenceViewer.cpp
@@ -600,8 +600,7 @@ void SReferenceViewer::ShowSizeMap()
 		UEdGraphNode_Reference* ReferenceNode = Cast<UEdGraphNode_Reference>(SelectedNodes.Array()[0]);
 		if ( ReferenceNode )
 		{
-			const FAssetData& AssetData = ReferenceNode->GetAssetData();
-			SelectedAssetPackageNames.AddUnique( AssetData.PackageName );
+			SelectedAssetPackageNames.AddUnique( ReferenceNode->GetPackageName() );
 		}
 	}
 
