Fixed the Reference Viewer "Make Collection with Referenced Assets" option failing

UE-18971 - ReferenceViewer_Make Collection with Referenced Asset Local Broken

This had fallen foul of the new naming restrictions for collections, and has been changed to use the leaf-name of the asset, rather than its full path.

This also adds the collection type icons to the sub-menu in the Reference Viewer.

diff --git a/Engine/Source/Editor/ReferenceViewer/Private/ReferenceViewerSchema.cpp b/Engine/Source/Editor/ReferenceViewer/Private/ReferenceViewerSchema.cpp
--- a/Engine/Source/Editor/ReferenceViewer/Private/ReferenceViewerSchema.cpp
+++ b/Engine/Source/Editor/ReferenceViewer/Private/ReferenceViewerSchema.cpp
@@ -82,7 +82,19 @@ FConnectionDrawingPolicy* UReferenceViewerSchema::CreateConnectionDrawingPolicy(
 
 void UReferenceViewerSchema::GetMakeCollectionWithReferencedAssetsSubMenu( FMenuBuilder& MenuBuilder )
 {
-	MenuBuilder.AddMenuEntry(FReferenceViewerActions::Get().MakeLocalCollectionWithReferencedAssets);
-	MenuBuilder.AddMenuEntry(FReferenceViewerActions::Get().MakePrivateCollectionWithReferencedAssets);
-	MenuBuilder.AddMenuEntry(FReferenceViewerActions::Get().MakeSharedCollectionWithReferencedAssets);
+	MenuBuilder.AddMenuEntry(FReferenceViewerActions::Get().MakeLocalCollectionWithReferencedAssets, 
+		NAME_None, TAttribute<FText>(), 
+		ECollectionShareType::GetDescription(ECollectionShareType::CST_Local), 
+		FSlateIcon(FEditorStyle::GetStyleSetName(), ECollectionShareType::GetIconStyleName(ECollectionShareType::CST_Local))
+		);
+	MenuBuilder.AddMenuEntry(FReferenceViewerActions::Get().MakePrivateCollectionWithReferencedAssets,
+		NAME_None, TAttribute<FText>(), 
+		ECollectionShareType::GetDescription(ECollectionShareType::CST_Private), 
+		FSlateIcon(FEditorStyle::GetStyleSetName(), ECollectionShareType::GetIconStyleName(ECollectionShareType::CST_Private))
+		);
+	MenuBuilder.AddMenuEntry(FReferenceViewerActions::Get().MakeSharedCollectionWithReferencedAssets,
+		NAME_None, TAttribute<FText>(), 
+		ECollectionShareType::GetDescription(ECollectionShareType::CST_Shared), 
+		FSlateIcon(FEditorStyle::GetStyleSetName(), ECollectionShareType::GetIconStyleName(ECollectionShareType::CST_Shared))
+		);
 }
diff --git a/Engine/Source/Editor/ReferenceViewer/Private/SReferenceViewer.cpp b/Engine/Source/Editor/ReferenceViewer/Private/SReferenceViewer.cpp
--- a/Engine/Source/Editor/ReferenceViewer/Private/SReferenceViewer.cpp
+++ b/Engine/Source/Editor/ReferenceViewer/Private/SReferenceViewer.cpp
@@ -585,7 +585,7 @@ void SReferenceViewer::MakeCollectionWithReferencedAssets(ECollectionShareType::
 
 	if ( SelectedObject )
 	{
-		FString DefaultName = FText::Format(NSLOCTEXT("UnrealEd", "Resources", "{0}_Resources"), FText::FromString( SelectedObject->GetPathName())).ToString();
+		FString DefaultName = FText::Format(NSLOCTEXT("UnrealEd", "Resources", "{0}_Resources"), FText::FromString(SelectedObject->GetName())).ToString();
 		ObjectTools::ShowReferencedObjs(SelectedObject, DefaultName, ShareType);
 	}
 }
