Editor: Added the ability to type in a path for the reference viewer, which makes it much easier to track down broken assets due to missing assets (paste in the missing asset path to view the assets that referenced it)
[INTEGRATE] Change 2712391 by [REDACTED_PERSON][REDACTED_HANDLE] on 2015/09/30 23:53:02

diff --git a/Engine/Source/Editor/EditorStyle/Private/SlateEditorStyle.cpp b/Engine/Source/Editor/EditorStyle/Private/SlateEditorStyle.cpp
--- a/Engine/Source/Editor/EditorStyle/Private/SlateEditorStyle.cpp
+++ b/Engine/Source/Editor/EditorStyle/Private/SlateEditorStyle.cpp
@@ -5605,6 +5605,10 @@ void FSlateEditorStyle::FStyle::SetupContentBrowserStyle()
 			.SetHighlightColor( FLinearColor( 1.0f, 1.0f, 1.0f ) )
 			.SetShadowOffset( FVector2D( 1,1 ) )
 			.SetShadowColorAndOpacity( FLinearColor(0,0,0,0.9f) ) );
+
+		Set("ReferenceViewer.PathText", FEditableTextBoxStyle(NormalEditableTextBoxStyle)
+			.SetFont(TTF_CORE_FONT("Fonts/Roboto-Bold", 11)));
+
 		Set( "ContentBrowser.Sources", new IMAGE_BRUSH( "ContentBrowser/sources_16x", Icon16x16 ) );
 		Set( "ContentBrowser.PathDelimiter", new IMAGE_BRUSH( "Common/SmallArrowRight", Icon10x10 ) );
 		Set( "ContentBrowser.LockButton_Locked", new IMAGE_BRUSH( "Icons/padlock_locked_16x", Icon16x16 ) );
diff --git a/Engine/Source/Editor/ReferenceViewer/Private/SReferenceViewer.cpp b/Engine/Source/Editor/ReferenceViewer/Private/SReferenceViewer.cpp
--- a/Engine/Source/Editor/ReferenceViewer/Private/SReferenceViewer.cpp
+++ b/Engine/Source/Editor/ReferenceViewer/Private/SReferenceViewer.cpp
@@ -112,9 +112,13 @@ void SReferenceViewer::Construct( const FArguments& InArgs )
 				SNew(SBorder)
 				.BorderImage( FEditorStyle::GetBrush("ToolPanel.GroupBorder") )
 				[
-					SNew(STextBlock)
+					SNew(SEditableTextBox)
 					.Text(this, &SReferenceViewer::GetAddressBarText)
-					.TextStyle(FEditorStyle::Get(), "ContentBrowser.PathText")
+					.OnTextCommitted(this, &SReferenceViewer::OnAddressBarTextCommitted)
+					.OnTextChanged(this, &SReferenceViewer::OnAddressBarTextChanged)
+					.SelectAllTextWhenFocused(true)
+					.SelectAllTextOnCommit(true)
+					.Style(FEditorStyle::Get(), "ReferenceViewer.PathText")
 				]
 			]
 		]
@@ -325,20 +329,47 @@ FText SReferenceViewer::GetAddressBarText() const
 {
 	if ( GraphObj )
 	{
-		const TArray<FName>& CurrentGraphRootPackageNames = GraphObj->GetCurrentGraphRootPackageNames();
-		if ( CurrentGraphRootPackageNames.Num() == 1 )
+		if (TemporaryPathBeingEdited.IsEmpty())
 		{
-			return FText::FromName(CurrentGraphRootPackageNames[0]);
+			const TArray<FName>& CurrentGraphRootPackageNames = GraphObj->GetCurrentGraphRootPackageNames();
+			if (CurrentGraphRootPackageNames.Num() == 1)
+			{
+				return FText::FromName(CurrentGraphRootPackageNames[0]);
+			}
+			else if (CurrentGraphRootPackageNames.Num() > 1)
+			{
+				return FText::Format(LOCTEXT("AddressBarMultiplePackagesText", "{0} and {1} others"), FText::FromName(CurrentGraphRootPackageNames[0]), FText::AsNumber(CurrentGraphRootPackageNames.Num()));
+			}
 		}
-		else if ( CurrentGraphRootPackageNames.Num() > 1 )
+		else
 		{
-			return FText::Format(LOCTEXT("AddressBarMultiplePackagesText", "{0} and {1} others"), FText::FromName(CurrentGraphRootPackageNames[0]), FText::AsNumber(CurrentGraphRootPackageNames.Num()));
+			return TemporaryPathBeingEdited;
 		}
 	}
 
 	return FText();
 }
 
+void SReferenceViewer::OnAddressBarTextCommitted(const FText& NewText, ETextCommit::Type CommitInfo)
+{
+	if (CommitInfo == ETextCommit::OnEnter)
+	{
+		const FName NewPath(*NewText.ToString());
+
+		TArray<FName> NewPaths;
+		NewPaths.Add(NewPath);
+
+		SetGraphRootPackageNames(NewPaths);
+	}
+
+	TemporaryPathBeingEdited = FText();
+}
+
+void SReferenceViewer::OnAddressBarTextChanged(const FText& NewText)
+{
+	TemporaryPathBeingEdited = NewText;
+}
+
 void SReferenceViewer::OnApplyHistoryData(const FReferenceViewerHistoryData& History)
 {
 	if ( GraphObj )
diff --git a/Engine/Source/Editor/ReferenceViewer/Private/SReferenceViewer.h b/Engine/Source/Editor/ReferenceViewer/Private/SReferenceViewer.h
--- a/Engine/Source/Editor/ReferenceViewer/Private/SReferenceViewer.h
+++ b/Engine/Source/Editor/ReferenceViewer/Private/SReferenceViewer.h
@@ -57,6 +57,12 @@ private:
 	/** Gets the text to be displayed in the address bar */
 	FText GetAddressBarText() const;
 
+	/** Called when the path is being edited */
+	void OnAddressBarTextChanged(const FText& NewText);
+
+	/** Sets the new path for the viewer */
+	void OnAddressBarTextCommitted(const FText& NewText, ETextCommit::Type CommitInfo);
+
 	void OnApplyHistoryData(const FReferenceViewerHistoryData& History);
 
 	void OnUpdateHistoryData(FReferenceViewerHistoryData& HistoryData) const;
@@ -98,4 +104,7 @@ private:
 	TSharedPtr<FUICommandList> ReferenceViewerActions;
 
 	UEdGraph_ReferenceViewer* GraphObj;
+
+	/** The temporary copy of the path text when it is actively being edited. */
+	FText TemporaryPathBeingEdited;
 };
