Copying //UE4/Dev-Framework to //UE4/Dev-Main (Source: //UE4/Dev-Framework @ 3182037)

diff --git a/Engine/Source/Editor/KismetCompiler/Private/KismetCompiler.cpp b/Engine/Source/Editor/KismetCompiler/Private/KismetCompiler.cpp
--- a/Engine/Source/Editor/KismetCompiler/Private/KismetCompiler.cpp
+++ b/Engine/Source/Editor/KismetCompiler/Private/KismetCompiler.cpp
@@ -169,12 +169,10 @@ void FKismetCompilerContext::FSubobjectCollection::AddObject(const UObject* cons
 	if ( InObject )
 	{
 		Collection.Add(InObject);
-		TArray<UObject*> Subobjects;
-		GetObjectsWithOuter(InObject, Subobjects, true);
-		for ( auto SubObject : Subobjects )
+		ForEachObjectWithOuter(InObject, [this](UObject* Child)
 		{
-			Collection.Add(SubObject);
-		}
+			Collection.Add(Child);
+		});
 	}
 }
 
@@ -236,7 +234,7 @@ void FKismetCompilerContext::CleanAndSanitizeClass(UBlueprintGeneratedClass* Cla
 	for( auto SubObjIt = ClassSubObjects.CreateIterator(); SubObjIt; ++SubObjIt )
 	{
 		UObject* CurrSubObj = *SubObjIt;
-		CurrSubObj->Rename(NULL, TransientClass, RenFlags);
+		CurrSubObj->Rename(nullptr, TransientClass, RenFlags);
 		if( UProperty* Prop = Cast<UProperty>(CurrSubObj) )
 		{
 			FKismetCompilerUtilities::InvalidatePropertyExport(Prop);
@@ -278,21 +276,21 @@ void FKismetCompilerContext::SaveSubObjectsFromCleanAndSanitizeClass(FSubobjectC
 
 	{
 		TSet<class UCurveBase*> Curves;
-		for ( auto Timeline : Blueprint->Timelines )
+		for ( UTimelineTemplate* Timeline : Blueprint->Timelines )
 		{
 			if ( Timeline )
 			{
 				Timeline->GetAllCurves(Curves);
 			}
 		}
-		for ( auto Component : Blueprint->ComponentTemplates )
+		for ( UActorComponent* Component : Blueprint->ComponentTemplates )
 		{
-			if ( auto TimelineComponent = Cast<UTimelineComponent>(Component) )
+			if ( UTimelineComponent* TimelineComponent = Cast<UTimelineComponent>(Component) )
 			{
 				TimelineComponent->GetAllCurves(Curves);
 			}
 		}
-		for ( auto Curve : Curves )
+		for ( UCurveBase* Curve : Curves )
 		{
 			SubObjectsToSave.AddObject(Curve);
 		}
@@ -313,17 +311,16 @@ void FKismetCompilerContext::PostCreateSchema()
 
 	TArray<UClass*> ClassesOfUK2Node;
 	GetDerivedClasses(UK2Node::StaticClass(), ClassesOfUK2Node, true);
-	for ( auto ClassIt = ClassesOfUK2Node.CreateConstIterator(); ClassIt; ++ClassIt )
+	for (UClass* Class : ClassesOfUK2Node)
 	{
-		if( !(*ClassIt)->HasAnyClassFlags(CLASS_Abstract) )
+		if( !Class->HasAnyClassFlags(CLASS_Abstract) )
 		{
-			UClass* Class = (*ClassIt);
 			UObject* CDO = Class->GetDefaultObject();
 			UK2Node* K2CDO = Class->GetDefaultObject<UK2Node>();
 			FNodeHandlingFunctor* HandlingFunctor = K2CDO->CreateNodeHandler(*this);
 			if (HandlingFunctor)
 			{
-				NodeHandlers.Add(*ClassIt, HandlingFunctor);
+				NodeHandlers.Add(Class, HandlingFunctor);
 			}
 		}
 	}
@@ -374,7 +371,7 @@ void FKismetCompilerContext::ValidatePin(const UEdGraphPin* Pin) const
 {
 	Super::ValidatePin(Pin);
 
-	auto OwningNodeUnchecked = Pin ? Pin->GetOwningNodeUnchecked() : NULL;
+	UEdGraphNode* OwningNodeUnchecked = Pin ? Pin->GetOwningNodeUnchecked() : nullptr;
 	if (!OwningNodeUnchecked)
 	{
 		//handled by Super::ValidatePin
@@ -499,7 +496,7 @@ void FKismetCompilerContext::ValidateVariableNames()
 			}
 			else if (ParentClass->IsNative()) // the above case handles when the parent is a blueprint
 			{
-				if (auto ExisingField = FindField<UField>(ParentClass, *VarNameStr))
+				if (UField* ExisingField = FindField<UField>(ParentClass, *VarNameStr))
 				{
 					UE_LOG(LogK2Compiler, Warning, TEXT("ValidateVariableNames name %s (used in %s) is already taken by %s")
 						, *VarNameStr, *Blueprint->GetPathName(), *ExisingField->GetPathName());
@@ -1580,7 +1577,7 @@ void FKismetCompilerContext::PrecompileFunction(FKismetFunctionContext& Context)
 		}
 		else
 		{
-			for (auto Node : Context.SourceGraph->Nodes)
+			for (UEdGraphNode* Node : Context.SourceGraph->Nodes)
 			{
 				const bool bGenerateParameters = Node->IsA<UK2Node_FunctionEntry>();
 				const bool bGenerateResult = Node->IsA<UK2Node_FunctionResult>();
@@ -1733,7 +1730,7 @@ void FKismetCompilerContext::CompileFunction(FKismetFunctionContext& Context)
 				}
 				if (Context.bInstrumentScriptCode)
 				{
-					for (auto Pin : Node->Pins)
+					for (UEdGraphPin* Pin : Node->Pins)
 					{
 						if (Pin->Direction == EGPD_Input && Pin->PinType.PinCategory == UEdGraphSchema_K2::PC_Exec)
 						{
@@ -1790,7 +1787,7 @@ void FKismetCompilerContext::CompileFunction(FKismetFunctionContext& Context)
 			const bool bDidNodeGenerateCode = SourceStatementList != nullptr && SourceStatementList->Num() > 0;
 			if (bDidNodeGenerateCode)
 			{
-				for (auto Statement : *SourceStatementList)
+				for (FBlueprintCompiledStatement* Statement : *SourceStatementList)
 				{
 					if (Statement && Statement->Type == KCST_InstrumentedPureNodeEntry)
 					{
@@ -1953,7 +1950,7 @@ void FKismetCompilerContext::FinishCompilingFunction(FKismetFunctionContext& Con
 	{
 		Function->SetMetaData(FBlueprintMetadata::MD_CallInEditor, TEXT( "true" ));
 	}
-	if (auto WorldContextPin = EntryNode->GetAutoWorldContextPin())
+	if (UEdGraphPin* WorldContextPin = EntryNode->GetAutoWorldContextPin())
 	{
 		Function->SetMetaData(FBlueprintMetadata::MD_WorldContext, *WorldContextPin->PinName);
 	}
@@ -2273,7 +2270,7 @@ void FKismetCompilerContext::ExpandTimelineNodes(UEdGraph* SourceGraph)
 		}
 	}
 	// Expand and validate timelines
-	for ( auto TimelinePair : Timelines )
+	for ( const FTimelinePair& TimelinePair : Timelines )
 	{
 		UK2Node_Timeline* const TimelineNode = TimelinePair.Node;
 		UTimelineTemplate* Timeline = TimelinePair.Template;
@@ -2468,7 +2465,7 @@ void FKismetCompilerContext::CreateFunctionStubForEvent(UK2Node_Event* SrcEventN
 
 	// Create the stub graph and add it to the list of functions to compile
 
-	auto ExistingGraph = static_cast<UObject*>(FindObjectWithOuter(OwnerOfTemporaries, UEdGraph::StaticClass(), EventNodeName));
+	UObject* ExistingGraph = static_cast<UObject*>(FindObjectWithOuter(OwnerOfTemporaries, UEdGraph::StaticClass(), EventNodeName));
 	if (ExistingGraph && !ExistingGraph->HasAnyFlags(RF_Transient))
 	{
 		MessageLog.Error(
@@ -2559,7 +2556,7 @@ void FKismetCompilerContext::CreateFunctionStubForEvent(UK2Node_Event* SrcEventN
 				}
 				else
 				{
-					auto VariableSetNode = SpawnIntermediateNode<UK2Node_VariableSet>(SrcEventNode, ChildStubGraph);
+					UK2Node_VariableSet* VariableSetNode = SpawnIntermediateNode<UK2Node_VariableSet>(SrcEventNode, ChildStubGraph);
 					VariableSetNode->VariableReference.SetSelfMember(NAME_None);
 					AssignmentNode = VariableSetNode;
 				}
@@ -2712,9 +2709,9 @@ void FKismetCompilerContext::ExpansionStep(UEdGraph* Graph, bool bAllowUbergraph
 							// We don't want to create boundary nodes around a set of pure only expansion nodes.
 							bImpureExpansionNodeAdded |= !IsNodePure(Graph->Nodes[ExpansionNodeIdx]);
 							// Add the relevent pins to the trace suppression filter
-							for (auto ExpansionPin : Graph->Nodes[ExpansionNodeIdx]->Pins)
+							for (UEdGraphPin* ExpansionPin : Graph->Nodes[ExpansionNodeIdx]->Pins)
 							{
-								for (auto ExpansionPinLink : ExpansionPin->LinkedTo)
+								for (UEdGraphPin* ExpansionPinLink : ExpansionPin->LinkedTo)
 								{
 									MessageLog.AddPinTraceFilter(ExpansionPin);
 								}
@@ -2746,12 +2743,12 @@ void FKismetCompilerContext::ExpansionStep(UEdGraph* Graph, bool bAllowUbergraph
 void FKismetCompilerContext::DetermineNodeExecLinks(UEdGraphNode* SourceNode, TMap<UEdGraphPin*, UEdGraphPin*>& SourceNodeLinks) const
 {
 	// Find all linked pins we care about from the source node.
-	for (auto SourcePin : SourceNode->Pins)
+	for (UEdGraphPin* SourcePin : SourceNode->Pins)
 	{
 		if (SourcePin->PinType.PinCategory == UEdGraphSchema_K2::PC_Exec)
 		{
 			UEdGraphPin* TrueSourcePin = MessageLog.FindSourcePin(SourcePin);
-			for (auto LinkedPin : SourcePin->LinkedTo)
+			for (UEdGraphPin* LinkedPin : SourcePin->LinkedTo)
 			{
 				SourceNodeLinks.Add(LinkedPin, TrueSourcePin);
 			}
@@ -3100,14 +3097,12 @@ void FKismetCompilerContext::ExpandTunnelsAndMacros(UEdGraph* SourceGraph)
 			// resolve any wildcard pins in the nodes cloned from the macro
 			if (!MacroInstanceNode->ResolvedWildcardType.PinCategory.IsEmpty())
 			{
-				for (auto ClonedNodeIt = ClonedGraph->Nodes.CreateConstIterator(); ClonedNodeIt; ++ClonedNodeIt)
+				for (UEdGraphNode* const ClonedNode : ClonedGraph->Nodes)
 				{
-					UEdGraphNode* const ClonedNode = *ClonedNodeIt;
 					if (ClonedNode)
 					{
-						for (auto ClonedPinIt = ClonedNode->Pins.CreateConstIterator(); ClonedPinIt; ++ClonedPinIt)
+						for (UEdGraphPin* const ClonedPin : ClonedNode->Pins)
 						{
-							UEdGraphPin* const ClonedPin = *ClonedPinIt;
 							if ( ClonedPin && (ClonedPin->PinType.PinCategory == Schema->PC_Wildcard) )
 							{
 								// copy only type info, so array or ref status is preserved
@@ -3115,13 +3110,13 @@ void FKismetCompilerContext::ExpandTunnelsAndMacros(UEdGraph* SourceGraph)
 								ClonedPin->PinType.PinSubCategory = MacroInstanceNode->ResolvedWildcardType.PinSubCategory;
 								ClonedPin->PinType.PinSubCategoryObject = MacroInstanceNode->ResolvedWildcardType.PinSubCategoryObject;
 							}
-							}
 						}
 					}
 				}
+			}
 
 			// Handle any nodes that need to inherit their macro instance's NodeGUID
-			for( auto ClonedNode : MacroNodes)
+			for( UEdGraphNode* ClonedNode : MacroNodes)
 			{
 				UK2Node_TemporaryVariable* TempVarNode = Cast<UK2Node_TemporaryVariable>(ClonedNode);
 				if(TempVarNode && TempVarNode->bIsPersistent)
@@ -3548,15 +3543,15 @@ void FKismetCompilerContext::Compile()
 	{
 		TArray<UEdGraph*> AllGraphs;
 		Blueprint->GetAllGraphs(AllGraphs);
-		for(auto GraphIter = AllGraphs.CreateIterator(); GraphIter; ++GraphIter)
+		for (UEdGraph* Graph : AllGraphs)
 		{
-			if(UEdGraph* Graph = *GraphIter)
+			if (Graph)
 			{
 				TArray<UK2Node*> AllNodes;
 				Graph->GetNodesOfClass(AllNodes);
-				for(auto NodeIter = AllNodes.CreateIterator(); NodeIter; ++NodeIter)
+				for (UK2Node* Node : AllNodes)
 				{
-					if(UK2Node* Node = *NodeIter)
+					if (Node)
 					{
 						Node->EarlyValidation(MessageLog);
 					}
@@ -3616,7 +3611,7 @@ void FKismetCompilerContext::Compile()
 
 	if (CompileOptions.CompileType == EKismetCompileType::Full)
 	{
-		auto InheritableComponentHandler = Blueprint->GetInheritableComponentHandler(false);
+		UInheritableComponentHandler* InheritableComponentHandler = Blueprint->GetInheritableComponentHandler(false);
 		if (InheritableComponentHandler)
 		{
 			InheritableComponentHandler->ValidateTemplates();
@@ -3664,7 +3659,7 @@ void FKismetCompilerContext::Compile()
 	{
 		//UBER GRAPH PERSISTENT FRAME
 		FEdGraphPinType Type(TEXT("struct"), TEXT(""), FPointerToUberGraphFrame::StaticStruct(), false, false);
-		auto Property = CreateVariable(UBlueprintGeneratedClass::GetUberGraphFrameName(), Type);
+		UProperty* Property = CreateVariable(UBlueprintGeneratedClass::GetUberGraphFrameName(), Type);
 		Property->SetPropertyFlags(CPF_DuplicateTransient | CPF_Transient);
 	}
 
@@ -3699,7 +3694,7 @@ void FKismetCompilerContext::Compile()
 		if (TargetClass->bHasInstrumentation)
 		{
 			FBlueprintDebugData& DebugData = TargetClass->GetDebugData();
-			for (auto EventInfo : SourcePinToExpansionEvent)
+			for (const TPair<UEdGraphPin*, UK2Node_Event*>& EventInfo : SourcePinToExpansionEvent)
 			{
 				DebugData.RegisterPinToEventName(EventInfo.Key, EventInfo.Value->GetFunctionName());
 			}
@@ -3751,19 +3746,19 @@ void FKismetCompilerContext::Compile()
 	{
 		TSet<UEdGraph*> AllGraphs;
 		AllGraphs.Add(UbergraphContext ? UbergraphContext->SourceGraph : NULL);
-		for(auto FunctionContextIter = FunctionList.CreateIterator(); FunctionContextIter; ++FunctionContextIter)
+		for (const FKismetFunctionContext& FunctionContext : FunctionList)
 		{
-			AllGraphs.Add((*FunctionContextIter).SourceGraph);
+			AllGraphs.Add(FunctionContext.SourceGraph);
 		}
-		for(auto GraphIter = AllGraphs.CreateIterator(); GraphIter; ++GraphIter)
+		for (UEdGraph* Graph : AllGraphs)
 		{
-			if(UEdGraph* Graph = *GraphIter)
+			if (Graph)
 			{
 				TArray<UK2Node_CreateDelegate*> AllNodes;
 				Graph->GetNodesOfClass(AllNodes);
-				for(auto NodeIter = AllNodes.CreateIterator(); NodeIter; ++NodeIter)
+				for (UK2Node_CreateDelegate* Node : AllNodes)
 				{
-					if(UK2Node_CreateDelegate* Node = *NodeIter)
+					if (Node)
 					{
 						Node->ValidationAfterFunctionsAreCreated(MessageLog, (0 != bIsFullCompile));
 					}
@@ -3948,15 +3943,15 @@ void FKismetCompilerContext::Compile()
 		{
 			UE_LOG(LogK2Compiler, Log, TEXT("\n\nLAYOUT CLASS %s:"), *GetNameSafe(NewClass));
 
-			for (auto Prop : TFieldRange<UProperty>(NewClass, EFieldIteratorFlags::ExcludeSuper))
+			for (UProperty* Prop : TFieldRange<UProperty>(NewClass, EFieldIteratorFlags::ExcludeSuper))
 			{
 				UE_LOG(LogK2Compiler, Log, TEXT("%5d:\t%-64s\t%s"), Prop->GetOffset_ForGC(), *GetNameSafe(Prop), *Prop->GetCPPType());
 			}
 
-			for (auto LocFunction : TFieldRange<UFunction>(NewClass, EFieldIteratorFlags::ExcludeSuper))
+			for (UFunction* LocFunction : TFieldRange<UFunction>(NewClass, EFieldIteratorFlags::ExcludeSuper))
 			{
 				UE_LOG(LogK2Compiler, Log, TEXT("\n\nLAYOUT FUNCTION %s:"), *GetNameSafe(LocFunction));
-				for (auto Prop : TFieldRange<UProperty>(LocFunction))
+				for (UProperty* Prop : TFieldRange<UProperty>(LocFunction))
 				{
 					const bool bOutParam = Prop && (0 != (Prop->PropertyFlags & CPF_OutParm));
 					const bool bInParam = Prop && !bOutParam && (0 != (Prop->PropertyFlags & CPF_Parm));
@@ -3972,14 +3967,10 @@ void FKismetCompilerContext::Compile()
 	// If this was a skeleton compile, make sure everything is RF_Transient
 	if (bIsSkeletonOnly)
 	{
-		TArray<UObject*> Subobjects;
-		GetObjectsWithOuter(NewClass, Subobjects);
-
-		for (auto SubObjIt = Subobjects.CreateIterator(); SubObjIt; ++SubObjIt)
+		ForEachObjectWithOuter(NewClass, [](UObject* Child)
 		{
-			UObject* CurrObj = *SubObjIt;
-			CurrObj->SetFlags(RF_Transient);
-		}
+			Child->SetFlags(RF_Transient);
+		});
 
 		NewClass->SetFlags(RF_Transient);
 
@@ -3992,7 +3983,7 @@ void FKismetCompilerContext::Compile()
 	{
 		TArray<UBlueprint*> DependentBlueprints;
 		FBlueprintEditorUtils::GetDependentBlueprints(Blueprint, DependentBlueprints);
-		for (auto CurrentBP : DependentBlueprints)
+		for (UBlueprint* CurrentBP : DependentBlueprints)
 		{
 			// Get the current dirty state of the package
 			UPackage* const Package = Cast<UPackage>(CurrentBP->GetOutermost());
@@ -4056,10 +4047,10 @@ void FKismetCompilerContext::Compile()
 			{
 				check(InProperty);
 
-				auto PropertyOwnerClass = InProperty->GetOwnerClass();
+				UClass* PropertyOwnerClass = InProperty->GetOwnerClass();
 				const bool bOwnerIsNativeClass = PropertyOwnerClass && PropertyOwnerClass->HasAnyClassFlags(CLASS_Native);
 
-				auto PropertyOwnerStruct = InProperty->GetOwnerStruct();
+				UStruct* PropertyOwnerStruct = InProperty->GetOwnerStruct();
 				const bool bOwnerIsNativeStruct = !PropertyOwnerClass && (!PropertyOwnerStruct || !PropertyOwnerStruct->IsA<UUserDefinedStruct>());
 
 				return InProperty->IsA<UStructProperty>()
@@ -4089,7 +4080,7 @@ void FKismetCompilerContext::Compile()
 		public:
 			static bool IsInnerProperty(const UObject* Object)
 			{
-				auto Property = Cast<const UProperty>(Object);
+				const UProperty* Property = Cast<const UProperty>(Object);
 				return Property // check arrays
 					&& Cast<const UFunction>(Property->GetOwnerStruct())
 					&& !Property->HasAnyPropertyFlags(CPF_Parm);
@@ -4102,7 +4093,7 @@ void FKismetCompilerContext::Compile()
 				if (Object && !IsInnerProperty(Object))
 				{
 					// Names of functions and properties are significant.
-					auto UniqueName = GetPathNameSafe(Object);
+					FString UniqueName = GetPathNameSafe(Object);
 					Ar << UniqueName;
 
 					if (Object->IsIn(RootObject))
@@ -4119,11 +4110,11 @@ void FKismetCompilerContext::Compile()
 				FArchive& Ar = *this;
 
 				bool bResult = false;
-				if (auto Struct = Cast<UStruct>(Object))
+				if (UStruct* Struct = Cast<UStruct>(Object))
 				{
 					if (Object == RootObject) // name and location are significant for the signature
 					{
-						auto UniqueName = GetPathNameSafe(Object);
+						FString UniqueName = GetPathNameSafe(Object);
 						Ar << UniqueName;
 					}
 
@@ -4131,12 +4122,12 @@ void FKismetCompilerContext::Compile()
 					Ar << SuperStruct;
 					Ar << Struct->Children;
 
-					if (auto Function = Cast<UFunction>(Struct))
+					if (UFunction* Function = Cast<UFunction>(Struct))
 					{
 						Ar << Function->FunctionFlags;
 					}
 
-					if (auto AsClass = Cast<UClass>(Struct))
+					if (UClass* AsClass = Cast<UClass>(Struct))
 					{
 						Ar << AsClass->ClassFlags;
 						Ar << AsClass->Interfaces;
@@ -4152,7 +4143,7 @@ void FKismetCompilerContext::Compile()
 		};
 
 		FSignatureArchiveCrc32 SignatureArchiveCrc32;
-		auto ParentBP = UBlueprint::GetBlueprintFromClass(NewClass->GetSuperClass());
+		UBlueprint* ParentBP = UBlueprint::GetBlueprintFromClass(NewClass->GetSuperClass());
 		const uint32 ParentSignatureCrc = ParentBP ? ParentBP->CrcLastCompiledSignature : 0;
 		Blueprint->CrcLastCompiledSignature = SignatureArchiveCrc32.Crc32(NewClass, ParentSignatureCrc);
 	}
diff --git a/Engine/Source/Editor/Persona/Private/AnimGraphNodeDetails.cpp b/Engine/Source/Editor/Persona/Private/AnimGraphNodeDetails.cpp
--- a/Engine/Source/Editor/Persona/Private/AnimGraphNodeDetails.cpp
+++ b/Engine/Source/Editor/Persona/Private/AnimGraphNodeDetails.cpp
@@ -25,6 +25,7 @@
 #include "Animation/AnimInstance.h"
 #include "ISkeletonEditorModule.h"
 #include "BlueprintEditor.h"
+#include "Animation/EditorAnimCurveBoneLinks.h"
 
 #define LOCTEXT_NAMESPACE "KismetNodeWithOptionalPinsDetails"
 
@@ -475,6 +476,8 @@ void FBoneReferenceCustomization::CustomizeHeader( TSharedRef<IPropertyHandle> S
 	UAnimationAsset * AnimationAsset = NULL;
 	USkeleton* TargetSkeleton = NULL;
 
+	TSharedPtr<IEditableSkeleton> EditableSkeleton; 
+
 	for (auto OuterIter = Objects.CreateIterator() ; OuterIter ; ++OuterIter)
 	{
 		AnimGraphNode = Cast<UAnimGraphNode_Base>(*OuterIter);
@@ -509,13 +512,23 @@ void FBoneReferenceCustomization::CustomizeHeader( TSharedRef<IPropertyHandle> S
 				break;
 			}
 		}
+
+		// editor animation curve bone links are responsible for linking joints to curve
+		// this is editor object that only exists for editor
+		if (UEditorAnimCurveBoneLinks* AnimCurveObj = Cast<UEditorAnimCurveBoneLinks>(*OuterIter))
+		{
+			EditableSkeleton = AnimCurveObj->EditableSkeleton.Pin();
+		}
 	}
 
 	if (TargetSkeleton)
 	{
 		ISkeletonEditorModule& SkeletonEditorModule = FModuleManager::LoadModuleChecked<ISkeletonEditorModule>("SkeletonEditor");
-		TSharedRef<IEditableSkeleton> EditableSkeleton = SkeletonEditorModule.CreateEditableSkeleton(TargetSkeleton);
+		EditableSkeleton = SkeletonEditorModule.CreateEditableSkeleton(TargetSkeleton);
+	}
 
+	if (EditableSkeleton.IsValid())
+	{
 		HeaderRow.NameContent()
 		[
 			StructPropertyHandle->CreatePropertyNameWidget()
@@ -523,12 +536,18 @@ void FBoneReferenceCustomization::CustomizeHeader( TSharedRef<IPropertyHandle> S
 
 		HeaderRow.ValueContent()
 		[
-			SNew(SBoneSelectionWidget, EditableSkeleton)
+			SNew(SBoneSelectionWidget, EditableSkeleton.ToSharedRef())
 			.Tooltip(StructPropertyHandle->GetToolTipText())
 			.OnBoneSelectionChanged(this, &FBoneReferenceCustomization::OnBoneSelectionChanged)
 			.OnGetSelectedBone(this, &FBoneReferenceCustomization::GetSelectedBone)
 		];
 	}
+	else
+	{
+		// if this FBoneReference is used by some other Outers, this will fail	
+		// should warn programmers instead of silent fail
+		ensureAlways(false);
+	}
 }
 END_SLATE_FUNCTION_BUILD_OPTIMIZATION
 
diff --git a/Engine/Source/Editor/Persona/Private/AnimationEditorViewportClient.cpp b/Engine/Source/Editor/Persona/Private/AnimationEditorViewportClient.cpp
--- a/Engine/Source/Editor/Persona/Private/AnimationEditorViewportClient.cpp
+++ b/Engine/Source/Editor/Persona/Private/AnimationEditorViewportClient.cpp
@@ -713,13 +713,7 @@ void FAnimationViewportClient::DisplayInfo(FCanvas* Canvas, FSceneView* View, bo
 
 			// Draw stats about the mesh
 			const FBoxSphereBounds& SkelBounds = PreviewMeshComponent->Bounds;
-			const FPlane ScreenPosition = View->Project(SkelBounds.Origin);
-
-			const int32 HalfX = Viewport->GetSizeXY().X / 2;
-			const int32 HalfY = Viewport->GetSizeXY().Y / 2;
-
-			const float ScreenRadius = FMath::Max((float)HalfX * View->ViewMatrices.GetProjectionMatrix().M[0][0], (float)HalfY * View->ViewMatrices.GetProjectionMatrix().M[1][1]) * SkelBounds.SphereRadius / FMath::Max(ScreenPosition.W, 1.0f);
-			const float LODFactor = ScreenRadius / 320.0f;
+			const float ScreenSize = ComputeBoundsScreenSize(SkelBounds.Origin, SkelBounds.SphereRadius, *View);
 
 			int32 NumBonesInUse;
 			int32 NumBonesMappedToVerts;
@@ -749,7 +743,7 @@ void FAnimationViewportClient::DisplayInfo(FCanvas* Canvas, FSceneView* View, bo
 
 			Canvas->DrawShadowedString(CurXOffset, CurYOffset, *InfoString, GEngine->GetSmallFont(), TextColor);
 
-			InfoString = FString::Printf(TEXT("Current Screen Size: %3.2f, FOV:%3.0f"), LODFactor, ViewFOV);
+			InfoString = FString::Printf(TEXT("Current Screen Size: %3.2f, FOV:%3.0f"), ScreenSize, ViewFOV);
 			CurYOffset += YL + 2;
 			Canvas->DrawShadowedString(CurXOffset, CurYOffset, *InfoString, GEngine->GetSmallFont(), TextColor);
 
@@ -838,16 +832,9 @@ void FAnimationViewportClient::DisplayInfo(FCanvas* Canvas, FSceneView* View, bo
 
 			// current screen size
 			const FBoxSphereBounds& SkelBounds = PreviewMeshComponent->Bounds;
-			const FPlane ScreenPosition = View->Project(SkelBounds.Origin);
-
-			const int32 HalfX = Viewport->GetSizeXY().X / 2;
-			const int32 HalfY = Viewport->GetSizeXY().Y / 2;
-			const float ScreenRadius = FMath::Max((float)HalfX * View->ViewMatrices.GetProjectionMatrix().M[0][0], (float)HalfY * View->ViewMatrices.GetProjectionMatrix().M[1][1]) * SkelBounds.SphereRadius / FMath::Max(ScreenPosition.W, 1.0f);
-			const float LODFactor = ScreenRadius / 320.0f;
-
-			float ScreenSize = ComputeBoundsScreenSize(ScreenPosition, SkelBounds.SphereRadius, *View);
+			const float ScreenSize = ComputeBoundsScreenSize(SkelBounds.Origin, SkelBounds.SphereRadius, *View);
 
-			InfoString = FString::Printf(TEXT("Current Screen Size: %3.2f"), LODFactor);
+			InfoString = FString::Printf(TEXT("Current Screen Size: %3.2f"), ScreenSize);
 			CurYOffset += YL + 2;
 			Canvas->DrawShadowedString(CurXOffset, CurYOffset, *InfoString, GEngine->GetSmallFont(), TextColor);
 
diff --git a/Engine/Source/Editor/Persona/Private/EditModes/SkeletonSelectionEditMode.cpp b/Engine/Source/Editor/Persona/Private/EditModes/SkeletonSelectionEditMode.cpp
--- a/Engine/Source/Editor/Persona/Private/EditModes/SkeletonSelectionEditMode.cpp
+++ b/Engine/Source/Editor/Persona/Private/EditModes/SkeletonSelectionEditMode.cpp
@@ -227,7 +227,7 @@ bool FSkeletonSelectionEditMode::InputDelta(FEditorViewportClient* InViewportCli
 				float RotAngle;
 				InRot.Quaternion().ToAxisAndAngle( RotAxis, RotAngle );
 
-				FVector4 BoneSpaceAxis = BaseTM.TransformVector( RotAxis );
+				FVector4 BoneSpaceAxis = BaseTM.TransformVectorNoScale( RotAxis );
 
 				//Calculate the new delta rotation
 				FQuat DeltaQuat( BoneSpaceAxis, RotAngle );
diff --git a/Engine/Source/Editor/Persona/Private/EditorObjectsTracker.cpp b/Engine/Source/Editor/Persona/Private/EditorObjectsTracker.cpp
--- a/Engine/Source/Editor/Persona/Private/EditorObjectsTracker.cpp
+++ b/Engine/Source/Editor/Persona/Private/EditorObjectsTracker.cpp
@@ -6,17 +6,26 @@
 void FEditorObjectTracker::AddReferencedObjects( FReferenceCollector& Collector )
 {
 	Collector.AddReferencedObjects(EditorObjMap);
+	Collector.AddReferencedObjects(EditorObjectArray);
 }
 
 UObject* FEditorObjectTracker::GetEditorObjectForClass( UClass* EdClass )
 {
-	UObject *Obj = (EditorObjMap.Contains(EdClass) ? *EditorObjMap.Find(EdClass) : NULL);
-	if(Obj == NULL)
+	UObject *Obj = (bAllowOnePerClass && EditorObjMap.Contains(EdClass) ? *EditorObjMap.Find(EdClass) : nullptr);
+
+	if (Obj == NULL)
 	{
 		FString ObjName = MakeUniqueObjectName(GetTransientPackage(), EdClass).ToString();
 		ObjName += "_EdObj";
 		Obj = NewObject<UObject>(GetTransientPackage(), EdClass, FName(*ObjName), RF_Public | RF_Standalone | RF_Transient);
-		EditorObjMap.Add(EdClass,Obj);
+		if (bAllowOnePerClass)
+		{
+			EditorObjMap.Add(EdClass, Obj);
+		}
+		else
+		{
+			EditorObjectArray.Add(Obj);
+		}
 	}
 	return Obj;
 }
\ No newline at end of file
diff --git a/Engine/Source/Editor/Persona/Private/EditorObjectsTracker.h b/Engine/Source/Editor/Persona/Private/EditorObjectsTracker.h
--- a/Engine/Source/Editor/Persona/Private/EditorObjectsTracker.h
+++ b/Engine/Source/Editor/Persona/Private/EditorObjectsTracker.h
@@ -8,6 +8,9 @@
 class FEditorObjectTracker : public FGCObject
 {
 public:
+	FEditorObjectTracker(bool bInAllowOnePerClass = true)
+		: bAllowOnePerClass(bInAllowOnePerClass)
+	{}
 
 	// FGCObject interface
 	void AddReferencedObjects( FReferenceCollector& Collector ) override;
@@ -17,8 +20,18 @@ public:
 	    if none exist */
 	UObject* GetEditorObjectForClass( UClass* EdClass );
 
+	void SetAllowOnePerClass(bool bInAllowOnePerClass)
+	{
+		bAllowOnePerClass = bInAllowOnePerClass;
+	}
 private:
 
+	/** If true, it uses TMap, otherwise, it just uses TArray */
+	bool bAllowOnePerClass;
+
 	/** Tracks editor objects created for details panel */
 	TMap< UClass*, UObject* >	EditorObjMap;
+
+	/** Tracks editor objects created for detail panel */
+	TArray<UObject*> EditorObjectArray;
 };
\ No newline at end of file
diff --git a/Engine/Source/Editor/Persona/Private/PersonaMeshDetails.cpp b/Engine/Source/Editor/Persona/Private/PersonaMeshDetails.cpp
--- a/Engine/Source/Editor/Persona/Private/PersonaMeshDetails.cpp
+++ b/Engine/Source/Editor/Persona/Private/PersonaMeshDetails.cpp
@@ -79,6 +79,143 @@ static void FillEnumOptions(TArray<TSharedPtr<FString> >& OutStrings, UEnum& InE
 	}
 }
 
+// Container widget for LOD buttons
+class SSkeletalLODActions : public SCompoundWidget
+{
+public:
+	SLATE_BEGIN_ARGS(SSkeletalLODActions)
+		: _LODIndex(INDEX_NONE)
+		, _PersonaToolkit(nullptr)
+	{}
+	SLATE_ARGUMENT(int32, LODIndex)
+	SLATE_ARGUMENT(TWeakPtr<IPersonaToolkit>, PersonaToolkit)
+	SLATE_EVENT(FOnClicked, OnRemoveLODClicked)
+	SLATE_EVENT(FOnClicked, OnReimportClicked)
+	SLATE_EVENT(FOnClicked, OnReimportNewFileClicked)
+
+	SLATE_END_ARGS()
+
+	void Construct(const FArguments& InArgs);
+
+private:
+
+	EActiveTimerReturnType RefreshExistFlag(double InCurrentTime, float InDeltaSeconds)
+	{
+		bDoesSourceFileExist_Cached = false;
+
+		TSharedPtr<IPersonaToolkit> SharedToolkit = PersonaToolkit.Pin();
+		if(SharedToolkit.IsValid())
+		{
+			USkeletalMesh* SkelMesh = SharedToolkit->GetMesh();
+
+			if(!SkelMesh)
+			{
+				return EActiveTimerReturnType::Continue;
+			}
+
+			if(SkelMesh->LODInfo.IsValidIndex(LODIndex))
+			{
+				FSkeletalMeshLODInfo& LODInfo = SkelMesh->LODInfo[LODIndex];
+
+				bDoesSourceFileExist_Cached = !LODInfo.SourceImportFilename.IsEmpty() && FPaths::FileExists(LODInfo.SourceImportFilename);
+			}
+		}
+		return EActiveTimerReturnType::Continue;
+	}
+
+	FText GetReimportButtonToolTipText() const
+	{
+		TSharedPtr<IPersonaToolkit> SharedToolkit = PersonaToolkit.Pin();
+
+		if(!CanReimportFromSource() || !SharedToolkit.IsValid())
+		{
+			return LOCTEXT("ReimportButton_NewFile_NoSource_ToolTip", "No source file available for reimport");
+		}
+
+		USkeletalMesh* SkelMesh = SharedToolkit->GetMesh();
+		check(SkelMesh);
+		check(SkelMesh->LODInfo.IsValidIndex(LODIndex)); // Should be true for the button to exist
+
+		FSkeletalMeshLODInfo& LODInfo = SkelMesh->LODInfo[LODIndex];
+
+		FString Filename = FPaths::GetCleanFilename(LODInfo.SourceImportFilename);
+
+		return FText::Format(LOCTEXT("ReimportButton_NewFile_ToolTip", "Reimport LOD{0} using the current source file ({1})"), FText::AsNumber(LODIndex), FText::FromString(Filename));
+	}
+
+	FText GetReimportButtonNewFileToolTipText() const
+	{
+		return FText::Format(LOCTEXT("ReimportButton_ToolTip", "Choose a new file to reimport over this LOD (LOD{0})"), FText::AsNumber(LODIndex));
+	}
+
+	bool CanReimportFromSource() const
+	{
+		return bDoesSourceFileExist_Cached;
+	}
+
+	// Incoming arg data
+	int32 LODIndex;
+	TWeakPtr<IPersonaToolkit> PersonaToolkit;
+	FOnClicked OnRemoveLODClicked;
+	FOnClicked OnReimportClicked;
+	FOnClicked OnReimportNewFileClicked;
+
+	// Cached exists flag so we don't constantly hit the file system
+	bool bDoesSourceFileExist_Cached;
+};
+
+void SSkeletalLODActions::Construct(const FArguments& InArgs)
+{
+	LODIndex = InArgs._LODIndex;
+	PersonaToolkit = InArgs._PersonaToolkit;
+	OnRemoveLODClicked = InArgs._OnRemoveLODClicked;
+	OnReimportClicked = InArgs._OnReimportClicked;
+	OnReimportNewFileClicked = InArgs._OnReimportNewFileClicked;
+
+	this->ChildSlot
+
+		[
+			SNew(SHorizontalBox)
+			+ SHorizontalBox::Slot()
+			.AutoWidth()
+			[
+				SNew(SButton)
+				.OnClicked(OnRemoveLODClicked)
+				[
+					SNew(STextBlock)
+					.Text(LOCTEXT("RemoveLOD", "Remove this LOD"))
+				]
+			]
+			+ SHorizontalBox::Slot()
+			.AutoWidth()
+			[
+				SNew(SButton)
+				.ToolTipText(this, &SSkeletalLODActions::GetReimportButtonToolTipText)
+				.IsEnabled(this, &SSkeletalLODActions::CanReimportFromSource)
+				.OnClicked(OnReimportClicked)
+				[
+					SNew(STextBlock)
+					.Text(LOCTEXT("ReimportLOD", "Reimport"))
+				]
+			]
+			+ SHorizontalBox::Slot()
+			.AutoWidth()
+			[
+				SNew(SButton)
+				.ToolTipText(this, &SSkeletalLODActions::GetReimportButtonNewFileToolTipText)
+				.OnClicked(OnReimportNewFileClicked)
+				[
+					SNew(STextBlock)
+					.Text(LOCTEXT("ReimportLOD_NewFile", "Reimport (New File)"))
+				]
+			]
+		];
+
+	// Register timer to refresh out exists flag periodically, with a bit added per LOD so we're not doing everything on the same frame
+	const float LODTimeOffset = 1.0f / 30.0f;
+	RegisterActiveTimer(1.0f + LODTimeOffset * LODIndex, FWidgetActiveTimerDelegate::CreateSP(this, &SSkeletalLODActions::RefreshExistFlag));
+}
+
 FSkelMeshReductionSettingsLayout::FSkelMeshReductionSettingsLayout(int32 InLODIndex, TSharedRef<FPersonaMeshDetails> InParentLODSettings, TSharedPtr<IPropertyHandle> InBoneToRemoveProperty)
 : LODIndex(InLODIndex)
 , ParentLODSettings(InParentLODSettings)
@@ -643,15 +780,14 @@ void FPersonaMeshDetails::AddLODLevelCategories(IDetailLayoutBuilder& DetailLayo
 			{
 				LODCategory.AddCustomRow(LOCTEXT("RemoveLODRow", "Remove LOD"))
 				.ValueContent()
-				.HAlign(HAlign_Left)
+				.HAlign(HAlign_Fill)
 				[
-					SNew(SButton)
-					.OnClicked(this, &FPersonaMeshDetails::RemoveOneLOD, LODIndex)
-					[
-						SNew(STextBlock)
-						.Text(LOCTEXT("RemoveLOD", "Remove this LOD"))
-						.Font(DetailLayout.GetDetailFont())
-					]
+					SNew(SSkeletalLODActions)
+					.LODIndex(LODIndex)
+					.PersonaToolkit(GetPersonaToolkit())
+					.OnRemoveLODClicked(this, &FPersonaMeshDetails::RemoveOneLOD, LODIndex)
+					.OnReimportClicked(this, &FPersonaMeshDetails::OnReimportLodClicked, &DetailLayout, EReimportButtonType::Reimport, LODIndex)
+					.OnReimportNewFileClicked(this, &FPersonaMeshDetails::OnReimportLodClicked, &DetailLayout, EReimportButtonType::ReimportWithNewFile, LODIndex)
 				];
 			}
 		}
@@ -667,10 +803,12 @@ void FPersonaMeshDetails::CustomizeLODSettingsCategories(IDetailLayoutBuilder& D
 
 	IDetailCategoryBuilder& LODSettingsCategory = DetailLayout.EditCategory("LodSettings", LOCTEXT("LodSettingsCategory", "LOD Settings"), ECategoryPriority::TypeSpecific);
 
+	TSharedPtr<SWidget> LodTextPtr;
+
 	LODSettingsCategory.AddCustomRow(LOCTEXT("LODImport", "LOD Import"))
 		.NameContent()
 		[
-			SNew(STextBlock)
+			SAssignNew(LodTextPtr, STextBlock)
 			.Font(IDetailLayoutBuilder::GetDetailFont())
 			.Text(LOCTEXT("LODImport", "LOD Import"))
 		]
@@ -1146,6 +1284,42 @@ void FPersonaMeshDetails::OnSetPostProcessBlueprint(const FAssetData& AssetData,
 	}
 }
 
+FReply FPersonaMeshDetails::OnReimportLodClicked(IDetailLayoutBuilder* DetailLayout, EReimportButtonType InReimportType, int32 InLODIndex)
+{
+	if(USkeletalMesh* SkelMesh = GetPersonaToolkit()->GetMesh())
+	{
+		if(!SkelMesh->LODInfo.IsValidIndex(InLODIndex))
+		{
+			return FReply::Unhandled();
+		}
+
+		FString SourceFilenameBackup("");
+		if(InReimportType == EReimportButtonType::ReimportWithNewFile)
+		{
+			// Back up current source filename and empty it so the importer asks for a new one.
+			SourceFilenameBackup = SkelMesh->LODInfo[InLODIndex].SourceImportFilename;
+			SkelMesh->LODInfo[InLODIndex].SourceImportFilename.Empty();
+		}
+
+		bool bImportSucceeded = FbxMeshUtils::ImportMeshLODDialog(SkelMesh, InLODIndex);
+
+		if(InReimportType == EReimportButtonType::ReimportWithNewFile && !bImportSucceeded)
+		{
+			// Copy old source file back, as this one failed
+			SkelMesh->LODInfo[InLODIndex].SourceImportFilename = SourceFilenameBackup;
+		}
+
+		if(DetailLayout)
+		{
+			DetailLayout->ForceRefreshDetails();
+		}
+
+		return FReply::Handled();
+	}
+
+	return FReply::Unhandled();
+}
+
 void FPersonaMeshDetails::OnGetMaterialsForArray(class IMaterialListBuilder& OutMaterials, int32 LODIndex)
 {
 	USkeletalMesh* SkelMesh = GetPersonaToolkit()->GetMesh();
diff --git a/Engine/Source/Editor/Persona/Private/PersonaMeshDetails.h b/Engine/Source/Editor/Persona/Private/PersonaMeshDetails.h
--- a/Engine/Source/Editor/Persona/Private/PersonaMeshDetails.h
+++ b/Engine/Source/Editor/Persona/Private/PersonaMeshDetails.h
@@ -360,6 +360,16 @@ private:
 
 	SVerticalBox::FSlot& GetUVDensitySlot(int32 MaterialIndex, int32 UVChannelIndex) const;
 
+	// Used to control the type of reimport to do with a named parameter
+	enum class EReimportButtonType : uint8
+	{
+		Reimport,
+		ReimportWithNewFile
+	};
+
+	// Handler for reimport buttons in LOD details
+	FReply OnReimportLodClicked(IDetailLayoutBuilder* DetailLayout, EReimportButtonType InReimportType, int32 InLODIndex);
+
 public:
 
 	bool IsApplyNeeded() const;
diff --git a/Engine/Source/Editor/Persona/Private/PersonaModule.cpp b/Engine/Source/Editor/Persona/Private/PersonaModule.cpp
--- a/Engine/Source/Editor/Persona/Private/PersonaModule.cpp
+++ b/Engine/Source/Editor/Persona/Private/PersonaModule.cpp
@@ -209,9 +209,9 @@ TSharedRef<class FWorkflowTabFactory> FPersonaModule::CreateAnimNotifiesTabFacto
 	return MakeShareable(new FSkeletonAnimNotifiesSummoner(InHostingApp, InEditableSkeleton, InOnChangeAnimNotifies, InOnPostUndo, InOnObjectsSelected));
 }
 
-TSharedRef<class FWorkflowTabFactory> FPersonaModule::CreateCurveViewerTabFactory(const TSharedRef<class FWorkflowCentricApplication>& InHostingApp, const TSharedRef<class IEditableSkeleton>& InEditableSkeleton, const TSharedRef<IPersonaPreviewScene>& InPreviewScene, FSimpleMulticastDelegate& InOnCurvesChanged, FSimpleMulticastDelegate& InOnPostUndo) const
+TSharedRef<class FWorkflowTabFactory> FPersonaModule::CreateCurveViewerTabFactory(const TSharedRef<class FWorkflowCentricApplication>& InHostingApp, const TSharedRef<class IEditableSkeleton>& InEditableSkeleton, const TSharedRef<IPersonaPreviewScene>& InPreviewScene, FSimpleMulticastDelegate& InOnCurvesChanged, FSimpleMulticastDelegate& InOnPostUndo, FOnObjectsSelected InOnObjectsSelected) const
 {
-	return MakeShareable(new FAnimCurveViewerTabSummoner(InHostingApp, InEditableSkeleton, InPreviewScene, InOnCurvesChanged, InOnPostUndo));
+	return MakeShareable(new FAnimCurveViewerTabSummoner(InHostingApp, InEditableSkeleton, InPreviewScene, InOnCurvesChanged, InOnPostUndo, InOnObjectsSelected));
 }
 
 TSharedRef<class FWorkflowTabFactory> FPersonaModule::CreateRetargetManagerTabFactory(const TSharedRef<class FWorkflowCentricApplication>& InHostingApp, const TSharedRef<class IEditableSkeleton>& InEditableSkeleton, const TSharedRef<IPersonaPreviewScene>& InPreviewScene, FSimpleMulticastDelegate& InOnPostUndo) const
diff --git a/Engine/Source/Editor/Persona/Private/SAnimCurveViewer.cpp b/Engine/Source/Editor/Persona/Private/SAnimCurveViewer.cpp
--- a/Engine/Source/Editor/Persona/Private/SAnimCurveViewer.cpp
+++ b/Engine/Source/Editor/Persona/Private/SAnimCurveViewer.cpp
@@ -17,6 +17,7 @@
 
 #include "GenericCommands.h"
 #include "CurveViewerCommands.h"
+#include "Animation/EditorAnimCurveBoneLinks.h"
 
 #define LOCTEXT_NAMESPACE "SAnimCurveViewer"
 
@@ -24,9 +25,9 @@ static const FName ColumnId_AnimCurveNameLabel( "Curve Name" );
 static const FName ColumnID_AnimCurveTypeLabel("Type");
 static const FName ColumnID_AnimCurveWeightLabel( "Weight" );
 static const FName ColumnID_AnimCurveEditLabel( "Edit" );
+static const FName ColumnID_AnimCurveNumBoneLabel("Num Bones");
 
 const float MaxMorphWeight = 5.f;
-
 //////////////////////////////////////////////////////////////////////////
 // SAnimCurveListRow
 
@@ -123,7 +124,7 @@ TSharedRef< SWidget > SAnimCurveListRow::GenerateWidgetForColumn( const FName& C
 				.OnValueCommitted( this, &SAnimCurveListRow::OnAnimCurveWeightValueCommitted )
 			];
 	}
-	else 
+	else if ( ColumnName == ColumnID_AnimCurveEditLabel)
 	{
 		return
 			SNew(SVerticalBox)
@@ -139,6 +140,32 @@ TSharedRef< SWidget > SAnimCurveListRow::GenerateWidgetForColumn( const FName& C
 				.IsChecked(this, &SAnimCurveListRow::IsAnimCurveAutoFillChangedChecked)
 			];
 	}
+	else
+	{
+		return
+			SNew(SVerticalBox)
+
+			+ SVerticalBox::Slot()
+			.AutoHeight()
+			.Padding(0.0f, 1.0f)
+			.VAlign(VAlign_Center)
+			.HAlign(HAlign_Center)
+			[
+				SNew(STextBlock)
+				.Text(this, &SAnimCurveListRow::GetNumConntectedBones)
+			];
+	}
+}
+
+FText SAnimCurveListRow::GetNumConntectedBones() const
+{
+	const FCurveMetaData* CurveMetaData = Item->EditableSkeleton.Pin()->GetSkeleton().GetCurveMetaData(Item->SmartName);
+	if (CurveMetaData)
+	{
+		return FText::AsNumber(CurveMetaData->LinkedBones.Num());
+	}
+
+	return FText::AsNumber(0);
 }
 
 TSharedRef< SWidget > SAnimCurveListRow::GetCurveTypeWidget()
@@ -490,8 +517,14 @@ FText SAnimCurveTypeList::GetAnimCurveType() const
 //////////////////////////////////////////////////////////////////////////
 // SAnimCurveViewer
 
-void SAnimCurveViewer::Construct(const FArguments& InArgs, const TSharedRef<class IEditableSkeleton>& InEditableSkeleton, const TSharedRef<IPersonaPreviewScene>& InPreviewScene, FSimpleMulticastDelegate& InOnCurvesChanged, FSimpleMulticastDelegate& InOnPostUndo)
+void SAnimCurveViewer::Construct(const FArguments& InArgs, const TSharedRef<class IEditableSkeleton>& InEditableSkeleton, const TSharedRef<IPersonaPreviewScene>& InPreviewScene, FSimpleMulticastDelegate& InOnCurvesChanged, FSimpleMulticastDelegate& InOnPostUndo, FOnObjectsSelected InOnObjectsSelected)
 {
+	OnObjectsSelected = InOnObjectsSelected;
+
+	bShowAllCurves = true;
+
+	EditorObjectTracker.SetAllowOnePerClass(false);
+
 	ContainerName = USkeleton::AnimCurveMappingName;
 
 	CachedPreviewInstance = nullptr;
@@ -555,6 +588,7 @@ void SAnimCurveViewer::Construct(const FArguments& InArgs, const TSharedRef<clas
 			.OnContextMenuOpening( this, &SAnimCurveViewer::OnGetContextMenuContent )
 			.ItemHeight( 22.0f )
 			.SelectionMode(ESelectionMode::Multi)
+			.OnSelectionChanged( this, &SAnimCurveViewer::OnSelectionChanged )
 			.HeaderRow
 			(
 				SNew( SHeaderRow )
@@ -563,7 +597,7 @@ void SAnimCurveViewer::Construct(const FArguments& InArgs, const TSharedRef<clas
 				.DefaultLabel( LOCTEXT( "AnimCurveNameLabel", "Curve Name" ) )
 
 				+ SHeaderRow::Column(ColumnID_AnimCurveTypeLabel)
-				.FillWidth(1.f)
+				.FillWidth(0.5f)
 				.DefaultLabel(LOCTEXT("AnimCurveTypeLabel", "Type"))
 
 				+ SHeaderRow::Column( ColumnID_AnimCurveWeightLabel )
@@ -573,6 +607,10 @@ void SAnimCurveViewer::Construct(const FArguments& InArgs, const TSharedRef<clas
 				+ SHeaderRow::Column(ColumnID_AnimCurveEditLabel)
 				.FillWidth(0.25f)
 				.DefaultLabel(LOCTEXT("AnimCurveEditLabel", "Auto"))
+
+				+ SHeaderRow::Column(ColumnID_AnimCurveNumBoneLabel)
+				.FillWidth(0.5f)
+				.DefaultLabel(LOCTEXT("AnimCurveNumBoneLabel", "Bones"))
 			)
 		]
 	];
@@ -595,6 +633,22 @@ SAnimCurveViewer::~SAnimCurveViewer()
 	}
 }
 
+bool SAnimCurveViewer::IsCurveFilterEnabled() const
+{
+	return !bShowAllCurves;
+}
+
+ECheckBoxState SAnimCurveViewer::IsShowingAllCurves() const
+{
+	return bShowAllCurves ? ECheckBoxState::Checked : ECheckBoxState::Unchecked;
+}
+
+void SAnimCurveViewer::OnToggleShowingAllCurves(ECheckBoxState NewState)
+{
+	bShowAllCurves = (NewState == ECheckBoxState::Checked);
+
+	RefreshCurveList();
+}
 FReply SAnimCurveViewer::OnKeyDown(const FGeometry& MyGeometry, const FKeyEvent& InKeyEvent)
 {
 	if (UICommandList.IsValid() && UICommandList->ProcessCommandBindings(InKeyEvent))
@@ -634,12 +688,6 @@ void SAnimCurveViewer::BindCommands()
 		FCanExecuteAction());
 }
 
-
-bool SAnimCurveViewer::IsCurveFilterEnabled() const
-{
-	return true;
-}
-
 void SAnimCurveViewer::RefreshCachePreviewInstance()
 {
 	if (CachedPreviewInstance && OnAddAnimationCurveDelegate.IsBound())
@@ -798,7 +846,7 @@ void SAnimCurveViewer::CreateAnimCurveList( const FString& SearchText )
 		TMap<FName, float> ActiveCurves;
 		UDebugSkelMeshComponent* MeshComponent = PreviewScenePtr.Pin()->GetPreviewMeshComponent();
 		UAnimInstance* AnimInstance = MeshComponent->GetAnimInstance();
-		if (AnimInstance)
+		if (!bShowAllCurves && AnimInstance)
 		{
 			// attribute curve should contain everything
 			// so only search other container if attribute is off
@@ -837,7 +885,7 @@ void SAnimCurveViewer::CreateAnimCurveList( const FString& SearchText )
 			}
 
 			// If we passed that, see if we are filtering to only active
-			if (bAddToList)
+			if (bAddToList && !bShowAllCurves)
 			{
 				bAddToList = ActiveCurves.Contains(SmartName.DisplayName);
 			}
@@ -848,7 +896,9 @@ void SAnimCurveViewer::CreateAnimCurveList( const FString& SearchText )
 				// If not already in list, add it
 				if (FindIndexOfAnimCurveInfo(AnimCurveList, SmartName.DisplayName) == INDEX_NONE)
 				{
-					TSharedRef<FDisplayedAnimCurveInfo> NewInfo = FDisplayedAnimCurveInfo::Make(EditableSkeletonPtr, ContainerName, SmartName);
+					UEditorAnimCurveBoneLinks* EditorMirrorObj = Cast<UEditorAnimCurveBoneLinks> (EditorObjectTracker.GetEditorObjectForClass(UEditorAnimCurveBoneLinks::StaticClass()));
+					EditorMirrorObj->Initialize(EditableSkeletonPtr, SmartName, FOnAnimCurveBonesChange::CreateSP(this, &SAnimCurveViewer::ApplyCurveBoneLinks));
+					TSharedRef<FDisplayedAnimCurveInfo> NewInfo = FDisplayedAnimCurveInfo::Make(EditableSkeletonPtr, ContainerName, SmartName, EditorMirrorObj);
 
 					// See if we have an override set, and if so, grab info
 					float Weight = 0.f;
@@ -888,6 +938,25 @@ void SAnimCurveViewer::CreateAnimCurveList( const FString& SearchText )
 
 void SAnimCurveViewer::CreateAnimCurveTypeList(TSharedRef<SHorizontalBox> HorizontalBox)
 {
+	// Add toggle button for 'all curves'
+	HorizontalBox->AddSlot()
+	.AutoWidth()
+	.Padding(3, 1)
+	[
+		SNew(SCheckBox)
+		.Style(FEditorStyle::Get(), "ToggleButtonCheckbox")
+		.ToolTipText(LOCTEXT("ShowAllCurvesTooltip", "Show all curves, or only active curves."))
+		.Type(ESlateCheckBoxType::ToggleButton)
+		.IsChecked(this, &SAnimCurveViewer::IsShowingAllCurves)
+		.OnCheckStateChanged(this, &SAnimCurveViewer::OnToggleShowingAllCurves)
+		.Padding(4)
+		.Content()
+		[
+			SNew(STextBlock)
+			.Text(LOCTEXT("ShowAllCurves", "All Curves"))
+		]
+	];
+
 	// Add check box for each curve type flag
 	TArray<int32> CurveFlagsToList;
 	CurveFlagsToList.Add(ACEF_DriveMorphTarget);
@@ -1027,5 +1096,43 @@ bool SAnimCurveViewer::CanDelete()
 	return AnimCurveListView->GetNumItemsSelected() > 0;
 }
 
+void SAnimCurveViewer::OnSelectionChanged(TSharedPtr<FDisplayedAnimCurveInfo> InItem, ESelectInfo::Type SelectInfo)
+{
+	// make sure the currently selected ones are refreshed if it's first time
+	TArray<UObject*> SelectedObjects;
+
+	TArray< TSharedPtr< FDisplayedAnimCurveInfo > > SelectedRows = AnimCurveListView->GetSelectedItems();
+	for (auto ItemIt = SelectedRows.CreateIterator(); ItemIt; ++ItemIt)
+	{
+		TSharedPtr< FDisplayedAnimCurveInfo > RowItem = (*ItemIt);
+		UEditorAnimCurveBoneLinks* EditorMirrorObj = RowItem->EditorMirrorObject;
+		if (RowItem == InItem)
+		{
+			// first time selected, refresh
+			TArray<FBoneReference> BoneLinks;
+			FSmartName CurrentName = RowItem->SmartName;
+			const FCurveMetaData* CurveMetaData = EditableSkeletonPtr.Pin()->GetSkeleton().GetCurveMetaData(CurrentName);
+			if (CurveMetaData)
+			{
+				BoneLinks = CurveMetaData->LinkedBones;
+			}
+
+			EditorMirrorObj->Refresh(CurrentName, BoneLinks);
+		}
+
+		SelectedObjects.Add(EditorMirrorObj);
+	}
+
+	OnObjectsSelected.ExecuteIfBound(SelectedObjects);
+}
+
+void SAnimCurveViewer::ApplyCurveBoneLinks(class UEditorAnimCurveBoneLinks* EditorObj)
+{
+	if (EditorObj)
+	{
+		EditableSkeletonPtr.Pin()->SetCurveMetaBoneLinks(EditorObj->CurveName, EditorObj->ConnectedBones);
+	}
+}
+
 #undef LOCTEXT_NAMESPACE
 
diff --git a/Engine/Source/Editor/Persona/Private/SAnimCurveViewer.h b/Engine/Source/Editor/Persona/Private/SAnimCurveViewer.h
--- a/Engine/Source/Editor/Persona/Private/SAnimCurveViewer.h
+++ b/Engine/Source/Editor/Persona/Private/SAnimCurveViewer.h
@@ -3,6 +3,7 @@
 #pragma once
 
 #include "Animation/AnimInstance.h"
+#include "EditorObjectsTracker.h"
 
 class SAnimCurveViewer;
 
@@ -18,21 +19,23 @@ public:
 	TWeakPtr<class IEditableSkeleton> EditableSkeleton;		// The skeleton we're associated with
 	TSharedPtr<SInlineEditableTextBlock> EditableText;	// The editable text box in the list, used to focus from the context menu
 	FName ContainerName;	// The container in the skeleton this name resides in
+	class UEditorAnimCurveBoneLinks* EditorMirrorObject;
 
 	/** Static function for creating a new item, but ensures that you can only have a TSharedRef to one */
-	static TSharedRef<FDisplayedAnimCurveInfo> Make(TWeakPtr<class IEditableSkeleton> InEditableSkeleton, const FName& InContainerName, const FSmartName& InSmartName)
+	static TSharedRef<FDisplayedAnimCurveInfo> Make(TWeakPtr<class IEditableSkeleton> InEditableSkeleton, const FName& InContainerName, const FSmartName& InSmartName, class UEditorAnimCurveBoneLinks* InEditorMirrorObject)
 	{
-		return MakeShareable(new FDisplayedAnimCurveInfo(InEditableSkeleton, InContainerName, InSmartName));
+		return MakeShareable(new FDisplayedAnimCurveInfo(InEditableSkeleton, InContainerName, InSmartName, InEditorMirrorObject));
 	}
 
 protected:
 	/** Hidden constructor, always use Make above */
-	FDisplayedAnimCurveInfo(TWeakPtr<class IEditableSkeleton> InEditableSkeleton, const FName& InContainerName, const FSmartName& InSmartName)
+	FDisplayedAnimCurveInfo(TWeakPtr<class IEditableSkeleton> InEditableSkeleton, const FName& InContainerName, const FSmartName& InSmartName, class UEditorAnimCurveBoneLinks* InEditorMirrorObject)
 		: SmartName(InSmartName)
 		, Weight( 0 )
 		, bAutoFillData(true)
 		, EditableSkeleton(InEditableSkeleton)
 		, ContainerName(InContainerName)
+		, EditorMirrorObject(InEditorMirrorObject)
 	{}
 };
 
@@ -107,6 +110,9 @@ private:
 	TWeakPtr<IPersonaPreviewScene> PreviewScenePtr;
 	/** Returns curve type widget constructed */
 	TSharedRef<SWidget> GetCurveTypeWidget();
+
+	/** returns display text for number of connected joint setting */
+	FText GetNumConntectedBones() const;
 };
 
 //////////////////////////////////////////////////////////////////////////
@@ -117,7 +123,7 @@ class SAnimCurveViewer : public SCompoundWidget
 public:
 	SLATE_BEGIN_ARGS( SAnimCurveViewer )
 	{}
-		
+	
 	SLATE_END_ARGS()
 
 	/**
@@ -126,7 +132,7 @@ public:
 	* @param InArgs - Arguments passed from Slate
 	*
 	*/
-	void Construct( const FArguments& InArgs, const TSharedRef<class IEditableSkeleton>& InEditableSkeleton, const TSharedRef<class IPersonaPreviewScene>& InPreviewScene, FSimpleMulticastDelegate& InOnCurvesChanged, FSimpleMulticastDelegate& InOnPostUndo );
+	void Construct( const FArguments& InArgs, const TSharedRef<class IEditableSkeleton>& InEditableSkeleton, const TSharedRef<class IPersonaPreviewScene>& InPreviewScene, FSimpleMulticastDelegate& InOnCurvesChanged, FSimpleMulticastDelegate& InOnPostUndo, FOnObjectsSelected InOnObjectsSelected);
 
 	/**
 	* Destructor - resets the animation curve
@@ -231,6 +237,7 @@ private:
 
 	/** Handler for context menus */
 	TSharedPtr<SWidget> OnGetContextMenuContent() const;
+	void OnSelectionChanged(TSharedPtr<FDisplayedAnimCurveInfo> InItem, ESelectInfo::Type SelectInfo);
 
 	/**
 	* Clears and rebuilds the table, according to an optional search string
@@ -290,6 +297,8 @@ private:
 
 	int32 CurrentCurveFlag;
 
+	bool bShowAllCurves;
+
 	TMap<FName, float> OverrideCurves;
 
 	/** Commands that are bound to delegates*/
@@ -300,4 +309,13 @@ private:
 
 	friend class SAnimCurveListRow;
 	friend class SAnimCurveTypeList;
+
+	/** Tracks objects created for the details panel */
+	FEditorObjectTracker EditorObjectTracker;
+
+	/** Delegate called to select objects */
+	FOnObjectsSelected OnObjectsSelected;
+
+	/** apply curve bone links from editor mirror object to skeleton */
+	void ApplyCurveBoneLinks(class UEditorAnimCurveBoneLinks* EditorObj);
 };
\ No newline at end of file
diff --git a/Engine/Source/Editor/Persona/Private/SAnimEditorBase.cpp b/Engine/Source/Editor/Persona/Private/SAnimEditorBase.cpp
--- a/Engine/Source/Editor/Persona/Private/SAnimEditorBase.cpp
+++ b/Engine/Source/Editor/Persona/Private/SAnimEditorBase.cpp
@@ -349,7 +349,7 @@ FText SAnimEditorBase::GetCurrentFrame() const
 	
 	if (UAnimSequenceBase* AnimSeqBase = Cast<UAnimSequenceBase>(GetEditorObject()))
 	{
-		LastFrame = AnimSeqBase->GetNumberOfFrames() - 1;
+		LastFrame = FMath::Max(AnimSeqBase->GetNumberOfFrames() - 1, 0);
 	}
 
 	static const FNumberFormattingOptions FractionNumberFormat = FNumberFormattingOptions()
diff --git a/Engine/Source/Editor/Persona/Private/SMorphTargetViewer.cpp b/Engine/Source/Editor/Persona/Private/SMorphTargetViewer.cpp
--- a/Engine/Source/Editor/Persona/Private/SMorphTargetViewer.cpp
+++ b/Engine/Source/Editor/Persona/Private/SMorphTargetViewer.cpp
@@ -188,6 +188,11 @@ void SMorphTargetListRow::OnMorphTargetAutoFillChecked(ECheckBoxState InState)
 		// clear value so that it can be filled up
 		MorphTargetViewer->AddMorphTargetOverride(Item->Name, 0.f, true);
 	}
+	else
+	{
+		// Setting value, add the override
+		MorphTargetViewer->AddMorphTargetOverride(Item->Name, Item->Weight, false);
+	}
 }
 
 ECheckBoxState SMorphTargetListRow::IsMorphTargetAutoFillChangedChecked() const
@@ -261,6 +266,8 @@ float SMorphTargetListRow::GetWeight() const
 { 
 	if (Item->bAutoFillData)
 	{
+		float CurrentWeight = 0.f;
+
 		USkeletalMeshComponent* SkelComp = PreviewScenePtr.Pin()->GetPreviewMeshComponent();
 		UAnimInstance* AnimInstance = (SkelComp) ? SkelComp->GetAnimInstance() : nullptr;
 		if (AnimInstance)
@@ -268,15 +275,18 @@ float SMorphTargetListRow::GetWeight() const
 			// make sure if they have value that's not same as saved value
 			TMap<FName, float> MorphCurves;
 			AnimInstance->GetAnimationCurveList(EAnimCurveType::MorphTargetCurve, MorphCurves);
-			const float* CurrentValue = MorphCurves.Find(Item->Name);
-			if (CurrentValue && *CurrentValue != 0.f)
+			const float* CurrentWeightPtr = MorphCurves.Find(Item->Name);
+			if (CurrentWeightPtr)
 			{
-				return *CurrentValue;
+				CurrentWeight = *CurrentWeightPtr;
 			}
 		}
+		return CurrentWeight;
+	}
+	else
+	{
+		return Item->Weight;
 	}
-
-	return Item->Weight; 
 }
 //////////////////////////////////////////////////////////////////////////
 // SMorphTargetViewer
diff --git a/Engine/Source/Editor/Persona/Private/Shared/PersonaMode.cpp b/Engine/Source/Editor/Persona/Private/Shared/PersonaMode.cpp
--- a/Engine/Source/Editor/Persona/Private/Shared/PersonaMode.cpp
+++ b/Engine/Source/Editor/Persona/Private/Shared/PersonaMode.cpp
@@ -75,7 +75,7 @@ FPersonaAppMode::FPersonaAppMode(TSharedPtr<class FPersona> InPersona, FName InM
 	PersonaTabFactories.RegisterFactory(MakeShareable(new FAnimationAssetBrowserSummoner(InPersona, InPersona->GetPersonaToolkit(), FOnOpenNewAsset::CreateSP(InPersona.Get(), &FPersona::HandleOpenNewAsset), FOnAnimationSequenceBrowserCreated::CreateSP(InPersona.Get(), &FPersona::HandleAnimationSequenceBrowserCreated), true)));
 	PersonaTabFactories.RegisterFactory(MakeShareable(new FPreviewViewportSummoner(InPersona, InPersona->GetSkeletonTree(), InPersona->GetPreviewScene(), InPersona->OnPostUndo, InPersona, FOnViewportCreated::CreateSP(InPersona.Get(), &FPersona::HandleViewportCreated), InModeName != FPersonaModes::AnimationEditMode, true)));
 	PersonaTabFactories.RegisterFactory(MakeShareable(new FMorphTargetTabSummoner(InPersona, InPersona->GetPreviewScene(), InPersona->OnPostUndo)));
-	PersonaTabFactories.RegisterFactory(MakeShareable(new FAnimCurveViewerTabSummoner(InPersona, InPersona->GetSkeletonTree()->GetEditableSkeleton(), InPersona->GetPreviewScene(), InPersona->OnCurvesChanged, InPersona->OnPostUndo)));
+	PersonaTabFactories.RegisterFactory(MakeShareable(new FAnimCurveViewerTabSummoner(InPersona, InPersona->GetSkeletonTree()->GetEditableSkeleton(), InPersona->GetPreviewScene(), InPersona->OnCurvesChanged, InPersona->OnPostUndo, FOnObjectsSelected::CreateSP(InPersona.Get(), &FPersona::HandleObjectsSelected))));
 	PersonaTabFactories.RegisterFactory(MakeShareable(new FSkeletonAnimNotifiesSummoner(InPersona, InPersona->GetSkeletonTree()->GetEditableSkeleton(), InPersona->OnAnimNotifiesChanged, InPersona->OnPostUndo, FOnObjectsSelected::CreateSP(InPersona.Get(), &FPersona::HandleObjectsSelected))));
 	PersonaTabFactories.RegisterFactory(MakeShareable(new FRetargetManagerTabSummoner(InPersona, InPersona->GetSkeletonTree()->GetEditableSkeleton(), InPersona->GetPreviewScene(), InPersona->OnPostUndo)));
 	PersonaTabFactories.RegisterFactory(MakeShareable(new FSkeletonSlotNamesSummoner(InPersona, InPersona->GetSkeletonTree()->GetEditableSkeleton(), InPersona->OnPostUndo, FOnObjectSelected::CreateSP(InPersona.Get(), &FPersona::HandleObjectSelected))));
@@ -204,12 +204,13 @@ TSharedRef<SWidget> FMorphTargetTabSummoner::CreateTabBody(const FWorkflowTabSpa
 
 #include "SAnimCurveViewer.h"
 
-FAnimCurveViewerTabSummoner::FAnimCurveViewerTabSummoner(TSharedPtr<class FAssetEditorToolkit> InHostingApp, const TSharedRef<IEditableSkeleton>& InEditableSkeleton, const TSharedRef<IPersonaPreviewScene>& InPreviewScene, FSimpleMulticastDelegate& InOnCurvesChanged, FSimpleMulticastDelegate& InOnPostUndo)
+FAnimCurveViewerTabSummoner::FAnimCurveViewerTabSummoner(TSharedPtr<class FAssetEditorToolkit> InHostingApp, const TSharedRef<IEditableSkeleton>& InEditableSkeleton, const TSharedRef<IPersonaPreviewScene>& InPreviewScene, FSimpleMulticastDelegate& InOnCurvesChanged, FSimpleMulticastDelegate& InOnPostUndo, FOnObjectsSelected InOnObjectsSelected)
 	: FWorkflowTabFactory(FPersonaTabs::AnimCurveViewID, InHostingApp)
 	, EditableSkeleton(InEditableSkeleton)
 	, PreviewScene(InPreviewScene)
 	, OnCurvesChanged(InOnCurvesChanged)
 	, OnPostUndo(InOnPostUndo)
+	, OnObjectsSelected(InOnObjectsSelected)
 {
 	TabLabel = LOCTEXT("AnimCurveViewTabTitle", "Anim Curves");
 	TabIcon = FSlateIcon(FEditorStyle::GetStyleSetName(), "Persona.Tabs.AnimCurvePreviewer");
@@ -223,7 +224,7 @@ FAnimCurveViewerTabSummoner::FAnimCurveViewerTabSummoner(TSharedPtr<class FAsset
 
 TSharedRef<SWidget> FAnimCurveViewerTabSummoner::CreateTabBody(const FWorkflowTabSpawnInfo& Info) const
 {
-	return SNew(SAnimCurveViewer, EditableSkeleton.Pin().ToSharedRef(), PreviewScene.Pin().ToSharedRef(), OnCurvesChanged, OnPostUndo);
+	return SNew(SAnimCurveViewer, EditableSkeleton.Pin().ToSharedRef(), PreviewScene.Pin().ToSharedRef(), OnCurvesChanged, OnPostUndo, OnObjectsSelected);
 }
 
 /////////////////////////////////////////////////////
diff --git a/Engine/Source/Editor/Persona/Private/Shared/PersonaMode.h b/Engine/Source/Editor/Persona/Private/Shared/PersonaMode.h
--- a/Engine/Source/Editor/Persona/Private/Shared/PersonaMode.h
+++ b/Engine/Source/Editor/Persona/Private/Shared/PersonaMode.h
@@ -197,7 +197,7 @@ private:
 struct FAnimCurveViewerTabSummoner : public FWorkflowTabFactory
 {
 public:
-	FAnimCurveViewerTabSummoner(TSharedPtr<class FAssetEditorToolkit> InHostingApp, const TSharedRef<IEditableSkeleton>& InEditableSkeleton, const TSharedRef<IPersonaPreviewScene>& InPreviewScene, FSimpleMulticastDelegate& InOnCurvesChanged, FSimpleMulticastDelegate& InOnPostUndo);
+	FAnimCurveViewerTabSummoner(TSharedPtr<class FAssetEditorToolkit> InHostingApp, const TSharedRef<IEditableSkeleton>& InEditableSkeleton, const TSharedRef<IPersonaPreviewScene>& InPreviewScene, FSimpleMulticastDelegate& InOnCurvesChanged, FSimpleMulticastDelegate& InOnPostUndo, FOnObjectsSelected InOnObjectsSelected);
 
 	virtual TSharedRef<SWidget> CreateTabBody(const FWorkflowTabSpawnInfo& Info) const override;
 
@@ -212,6 +212,7 @@ private:
 	TWeakPtr<class IPersonaPreviewScene> PreviewScene;
 	FSimpleMulticastDelegate& OnCurvesChanged;
 	FSimpleMulticastDelegate& OnPostUndo;
+	FOnObjectsSelected OnObjectsSelected;
 };
 
 
diff --git a/Engine/Source/Editor/Persona/Public/PersonaModule.h b/Engine/Source/Editor/Persona/Public/PersonaModule.h
--- a/Engine/Source/Editor/Persona/Public/PersonaModule.h
+++ b/Engine/Source/Editor/Persona/Public/PersonaModule.h
@@ -130,7 +130,7 @@ public:
 	virtual TSharedRef<class FWorkflowTabFactory> CreateAnimNotifiesTabFactory(const TSharedRef<class FWorkflowCentricApplication>& InHostingApp, const TSharedRef<class IEditableSkeleton>& InEditableSkeleton, FSimpleMulticastDelegate& InOnChangeAnimNotifies, FSimpleMulticastDelegate& InOnPostUndo, FOnObjectsSelected InOnObjectsSelected) const;
 
 	/** Create a skeleton cuve viewer tab factory */
-	virtual TSharedRef<class FWorkflowTabFactory> CreateCurveViewerTabFactory(const TSharedRef<class FWorkflowCentricApplication>& InHostingApp, const TSharedRef<class IEditableSkeleton>& InEditableSkeleton, const TSharedRef<class IPersonaPreviewScene>& InPreviewScene, FSimpleMulticastDelegate& InOnCurvesChanged, FSimpleMulticastDelegate& InOnPostUndo) const;
+	virtual TSharedRef<class FWorkflowTabFactory> CreateCurveViewerTabFactory(const TSharedRef<class FWorkflowCentricApplication>& InHostingApp, const TSharedRef<class IEditableSkeleton>& InEditableSkeleton, const TSharedRef<class IPersonaPreviewScene>& InPreviewScene, FSimpleMulticastDelegate& InOnCurvesChanged, FSimpleMulticastDelegate& InOnPostUndo, FOnObjectsSelected InOnObjectsSelected) const;
 
 	/** Create a retarget manager tab factory */
 	virtual TSharedRef<class FWorkflowTabFactory> CreateRetargetManagerTabFactory(const TSharedRef<class FWorkflowCentricApplication>& InHostingApp, const TSharedRef<class IEditableSkeleton>& InEditableSkeleton, const TSharedRef<IPersonaPreviewScene>& InPreviewScene, FSimpleMulticastDelegate& InOnPostUndo) const;
diff --git a/Engine/Source/Editor/ReferenceViewer/Classes/EdGraphNode_Reference.h b/Engine/Source/Editor/ReferenceViewer/Classes/EdGraphNode_Reference.h
--- a/Engine/Source/Editor/ReferenceViewer/Classes/EdGraphNode_Reference.h
+++ b/Engine/Source/Editor/ReferenceViewer/Classes/EdGraphNode_Reference.h
@@ -9,15 +9,20 @@ class UEdGraphNode_Reference : public UEdGraphNode
 {
 	GENERATED_UCLASS_BODY()
 
-	virtual void SetupReferenceNode(const FIntPoint& NodeLoc, const TArray<FName>& NewPackageNames, const FAssetData& InAssetData);
+	virtual void SetupReferenceNode(const FIntPoint& NodeLoc, const TArray<FAssetIdentifier>& NewIdentifiers, const FAssetData& InAssetData);
 	virtual void SetReferenceNodeCollapsed(const FIntPoint& NodeLoc, int32 InNumReferencesExceedingMax);
 	virtual void AddReferencer(class UEdGraphNode_Reference* ReferencerNode);
-	virtual FName GetPackageName() const;
+	// Returns first asset identifier
+	virtual FAssetIdentifier GetIdentifier() const;
+	virtual void GetAllIdentifiers(TArray<FAssetIdentifier>& OutIdentifiers) const;
+	// Returns only the packages in this node, skips searchable names
 	virtual void GetAllPackageNames(TArray<FName>& OutPackageNames) const;
 	class UEdGraph_ReferenceViewer* GetReferenceViewerGraph() const;
 
 	// UEdGraphNode implementation
 	virtual FText GetNodeTitle(ENodeTitleType::Type TitleType) const override;
+	FLinearColor GetNodeTitleColor() const override;
+	virtual FText GetTooltipText() const override;
 	virtual void AllocateDefaultPins() override;
 	virtual UObject* GetJumpTargetForDoubleClick() const override;
 	// End UEdGraphNode implementation
@@ -25,6 +30,8 @@ class UEdGraphNode_Reference : public UEdGraphNode
 	void CacheAssetData(const FAssetData& AssetData);
 
 	bool UsesThumbnail() const;
+	bool IsPackage() const;
+	bool IsCollapsed() const;
 	FAssetData GetAssetData() const;
 
 	virtual UEdGraphPin* GetDependencyPin();
@@ -36,10 +43,12 @@ protected:
 	}
 
 private:
-	TArray<FName> PackageNames;
+	TArray<FAssetIdentifier> Identifiers;
 	FText NodeTitle;
 
 	bool bUsesThumbnail;
+	bool bIsPackage;
+	bool bIsCollapsed;
 	FAssetData CachedAssetData;
 
 	UEdGraphPin* DependencyPin;
diff --git a/Engine/Source/Editor/ReferenceViewer/Classes/EdGraph_ReferenceViewer.h b/Engine/Source/Editor/ReferenceViewer/Classes/EdGraph_ReferenceViewer.h
--- a/Engine/Source/Editor/ReferenceViewer/Classes/EdGraph_ReferenceViewer.h
+++ b/Engine/Source/Editor/ReferenceViewer/Classes/EdGraph_ReferenceViewer.h
@@ -2,6 +2,7 @@
 
 #pragma once
 
+#include "AssetRegistryModule.h"
 #include "EdGraph_ReferenceViewer.generated.h"
 
 UCLASS()
@@ -14,23 +15,23 @@ public:
 	virtual void BeginDestroy() override;
 	// End UObject implementation
 
-	void SetGraphRoot(const TArray<FName>& GraphRootPackageNames, const FIntPoint& GraphRootOrigin = FIntPoint(ForceInitToZero));
-	const TArray<FName>& GetCurrentGraphRootPackageNames() const;
+	void SetGraphRoot(const TArray<FAssetIdentifier>& GraphRootIdentifiers, const FIntPoint& GraphRootOrigin = FIntPoint(ForceInitToZero));
+	const TArray<FAssetIdentifier>& GetCurrentGraphRootIdentifiers() const;
 	class UEdGraphNode_Reference* RebuildGraph();
 
 	bool IsSearchDepthLimited() const;
 	bool IsSearchBreadthLimited() const;
 	bool IsShowSoftReferences() const;
-	bool IsShowSoftDependencies() const;
 	bool IsShowHardReferences() const;
-	bool IsShowHardDependencies() const;
+	bool IsShowSearchableNames() const;
+	bool IsShowNativePackages() const;
 
 	void SetSearchDepthLimitEnabled(bool newEnabled);
 	void SetSearchBreadthLimitEnabled(bool newEnabled);
 	void SetShowSoftReferencesEnabled(bool newEnabled);
-	void SetShowSoftDependenciesEnabled(bool newEnabled);
 	void SetShowHardReferencesEnabled(bool newEnabled);
-	void SetShowHardDependenciesEnabled(bool newEnabled);
+	void SetShowSearchableNames(bool newEnabled);
+	void SetShowNativePackages(bool newEnabled);
 
 	int32 GetSearchDepthLimit() const;
 	int32 GetSearchBreadthLimit() const;
@@ -42,13 +43,14 @@ public:
 	const TSharedPtr<class FAssetThumbnailPool>& GetAssetThumbnailPool() const;
 
 private:
-	UEdGraphNode_Reference* ConstructNodes(const TArray<FName>& GraphRootPackageNames, const FIntPoint& GraphRootOrigin);
-	int32 RecursivelyGatherSizes(bool bReferencers, const TArray<FName>& PackageNames, int32 CurrentDepth, TSet<FName>& VisitedNames, TMap<FName, int32>& OutNodeSizes) const;
+	UEdGraphNode_Reference* ConstructNodes(const TArray<FAssetIdentifier>& GraphRootIdentifiers, const FIntPoint& GraphRootOrigin);
+	int32 RecursivelyGatherSizes(bool bReferencers, const TArray<FAssetIdentifier>& Identifiers, int32 CurrentDepth, TSet<FAssetIdentifier>& VisitedNames, TMap<FAssetIdentifier, int32>& OutNodeSizes) const;
 	void GatherAssetData(const TSet<FName>& AllPackageNames, TMap<FName, FAssetData>& OutPackageToAssetDataMap) const;
-	class UEdGraphNode_Reference* RecursivelyConstructNodes(bool bReferencers, UEdGraphNode_Reference* RootNode, const TArray<FName>& PackageNames, const FIntPoint& NodeLoc, const TMap<FName, int32>& NodeSizes, const TMap<FName, FAssetData>& PackagesToAssetDataMap, int32 CurrentDepth, TSet<FName>& VisitedNames);
+	class UEdGraphNode_Reference* RecursivelyConstructNodes(bool bReferencers, UEdGraphNode_Reference* RootNode, const TArray<FAssetIdentifier>& Identifiers, const FIntPoint& NodeLoc, const TMap<FAssetIdentifier, int32>& NodeSizes, const TMap<FName, FAssetData>& PackagesToAssetDataMap, int32 CurrentDepth, TSet<FAssetIdentifier>& VisitedNames);
 
 	bool ExceedsMaxSearchDepth(int32 Depth) const;
 	bool ExceedsMaxSearchBreadth(int32 Breadth) const;
+	EAssetRegistryDependencyType::Type GetReferenceSearchFlags(bool bReferencers) const;
 
 	UEdGraphNode_Reference* CreateReferenceNode();
 
@@ -59,7 +61,7 @@ private:
 	/** Pool for maintaining and rendering thumbnails */
 	TSharedPtr<class FAssetThumbnailPool> AssetThumbnailPool;
 
-	TArray<FName> CurrentGraphRootPackageNames;
+	TArray<FAssetIdentifier> CurrentGraphRootIdentifiers;
 	FIntPoint CurrentGraphRootOrigin;
 
 	int32 MaxSearchDepth;
@@ -68,8 +70,8 @@ private:
 	bool bLimitSearchDepth;
 	bool bLimitSearchBreadth;
 	bool bIsShowSoftReferences;
-	bool bIsShowSoftDependencies;
 	bool bIsShowHardReferences;
-	bool bIsShowHardDependencies;
+	bool bIsShowSearchableNames;
+	bool bIsShowNativePackages;
 };
 
diff --git a/Engine/Source/Editor/ReferenceViewer/Private/EdGraphNode_Reference.cpp b/Engine/Source/Editor/ReferenceViewer/Private/EdGraphNode_Reference.cpp
--- a/Engine/Source/Editor/ReferenceViewer/Private/EdGraphNode_Reference.cpp
+++ b/Engine/Source/Editor/ReferenceViewer/Private/EdGraphNode_Reference.cpp
@@ -14,26 +14,42 @@ UEdGraphNode_Reference::UEdGraphNode_Reference(const FObjectInitializer& ObjectI
 {
 	DependencyPin = NULL;
 	ReferencerPin = NULL;
+	bIsCollapsed = false;
+	bIsPackage = false;
+	bUsesThumbnail = false;
 }
 
-void UEdGraphNode_Reference::SetupReferenceNode(const FIntPoint& NodeLoc, const TArray<FName>& NewPackageNames, const FAssetData& InAssetData)
+void UEdGraphNode_Reference::SetupReferenceNode(const FIntPoint& NodeLoc, const TArray<FAssetIdentifier>& NewIdentifiers, const FAssetData& InAssetData)
 {
-	check(NewPackageNames.Num() > 0);
+	check(NewIdentifiers.Num() > 0);
 
 	NodePosX = NodeLoc.X;
 	NodePosY = NodeLoc.Y;
 
-	PackageNames = NewPackageNames;
-	FString ShortPackageName = FPackageName::GetLongPackageAssetName(NewPackageNames[0].ToString());
-	if ( NewPackageNames.Num() == 1 )
+	Identifiers = NewIdentifiers;
+	const FAssetIdentifier& First = NewIdentifiers[0];
+	FString ShortPackageName = FPackageName::GetLongPackageAssetName(First.PackageName.ToString());
+
+	bIsCollapsed = false;
+	bIsPackage = true;
+	if (NewIdentifiers[0].IsValue())
 	{
-		NodeComment = NewPackageNames[0].ToString();
+		ShortPackageName = FString::Printf(TEXT("%s::%s"), *First.ObjectName.ToString(), *First.ValueName.ToString());
+		bIsPackage = false;
+	}
+
+	if (NewIdentifiers.Num() == 1 )
+	{
+		if (bIsPackage)
+		{
+			NodeComment = First.PackageName.ToString();
+		}
 		NodeTitle = FText::FromString(ShortPackageName);
 	}
 	else
 	{
-		NodeComment = FText::Format(LOCTEXT("ReferenceNodeMultiplePackagesTitle", "{0} nodes"), FText::AsNumber(NewPackageNames.Num())).ToString();
-		NodeTitle = FText::Format(LOCTEXT("ReferenceNodeMultiplePackagesComment", "{0} and {1} others"), FText::FromString(ShortPackageName), FText::AsNumber(NewPackageNames.Num() - 1));
+		NodeComment = FText::Format(LOCTEXT("ReferenceNodeMultiplePackagesTitle", "{0} nodes"), FText::AsNumber(NewIdentifiers.Num())).ToString();
+		NodeTitle = FText::Format(LOCTEXT("ReferenceNodeMultiplePackagesComment", "{0} and {1} others"), FText::FromString(ShortPackageName), FText::AsNumber(NewIdentifiers.Num() - 1));
 	}
 	
 	CacheAssetData(InAssetData);
@@ -45,7 +61,9 @@ void UEdGraphNode_Reference::SetReferenceNodeCollapsed(const FIntPoint& NodeLoc,
 	NodePosX = NodeLoc.X;
 	NodePosY = NodeLoc.Y;
 
-	PackageNames.Empty();
+	Identifiers.Empty();
+	bIsCollapsed = true;
+	bUsesThumbnail = false;
 	NodeComment = FText::Format(LOCTEXT("ReferenceNodeCollapsedMessage", "{0} other nodes"), FText::AsNumber(InNumReferencesExceedingMax)).ToString();
 
 	NodeTitle = LOCTEXT("ReferenceNodeCollapsedTitle", "Collapsed nodes");
@@ -65,19 +83,30 @@ void UEdGraphNode_Reference::AddReferencer(UEdGraphNode_Reference* ReferencerNod
 	}
 }
 
-FName UEdGraphNode_Reference::GetPackageName() const
+FAssetIdentifier UEdGraphNode_Reference::GetIdentifier() const
 {
-	if ( PackageNames.Num() > 0 )
+	if (Identifiers.Num() > 0)
 	{
-		return PackageNames[0];
+		return Identifiers[0];
 	}
-	
-	return NAME_None;
+
+	return FAssetIdentifier();
+}
+
+void UEdGraphNode_Reference::GetAllIdentifiers(TArray<FAssetIdentifier>& OutIdentifiers) const
+{
+	OutIdentifiers.Append(Identifiers);
 }
 
 void UEdGraphNode_Reference::GetAllPackageNames(TArray<FName>& OutPackageNames) const
 {
-	OutPackageNames.Append(PackageNames);
+	for (const FAssetIdentifier& AssetId : Identifiers)
+	{
+		if (AssetId.PackageName.IsValid() && !AssetId.IsValue())
+		{
+			OutPackageNames.AddUnique(AssetId.PackageName);
+		}
+	}
 }
 
 UEdGraph_ReferenceViewer* UEdGraphNode_Reference::GetReferenceViewerGraph() const
@@ -90,6 +119,36 @@ FText UEdGraphNode_Reference::GetNodeTitle(ENodeTitleType::Type TitleType) const
 	return NodeTitle;
 }
 
+FLinearColor UEdGraphNode_Reference::GetNodeTitleColor() const
+{
+	if (bIsPackage)
+	{
+		return FLinearColor(0.4f, 0.62f, 1.0f);
+	}
+	else if (bIsCollapsed)
+	{
+		return FLinearColor(0.55f, 0.55f, 0.55f);
+	}
+	else 
+	{
+		return FLinearColor(0.0f, 0.55f, 0.62f);
+	}
+}
+
+FText UEdGraphNode_Reference::GetTooltipText() const
+{
+	FString TooltipString;
+	for (const FAssetIdentifier& AssetId : Identifiers)
+	{
+		if (!TooltipString.IsEmpty())
+		{
+			TooltipString.Append(TEXT("\n"));
+		}
+		TooltipString.Append(AssetId.ToString());
+	}
+	return FText::FromString(TooltipString);
+}
+
 void UEdGraphNode_Reference::AllocateDefaultPins()
 {
 	ReferencerPin = CreatePin( EEdGraphPinDirection::EGPD_Input, TEXT(""), TEXT(""), NULL, false, false, TEXT("") );
@@ -101,9 +160,9 @@ void UEdGraphNode_Reference::AllocateDefaultPins()
 
 UObject* UEdGraphNode_Reference::GetJumpTargetForDoubleClick() const
 {
-	if ( PackageNames.Num() > 0 )
+	if (Identifiers.Num() > 0 )
 	{
-		GetReferenceGraph()->SetGraphRoot(PackageNames, FIntPoint(NodePosX, NodePosY));
+		GetReferenceGraph()->SetGraphRoot(Identifiers, FIntPoint(NodePosX, NodePosY));
 		GetReferenceGraph()->RebuildGraph();
 	}
 	return NULL;
@@ -121,7 +180,7 @@ UEdGraphPin* UEdGraphNode_Reference::GetReferencerPin()
 
 void UEdGraphNode_Reference::CacheAssetData(const FAssetData& AssetData)
 {
-	if ( AssetData.IsValid() )
+	if ( AssetData.IsValid() && IsPackage() )
 	{
 		bUsesThumbnail = true;
 		CachedAssetData = AssetData;
@@ -131,9 +190,9 @@ void UEdGraphNode_Reference::CacheAssetData(const FAssetData& AssetData)
 		CachedAssetData = FAssetData();
 		bUsesThumbnail = false;
 
-		if ( PackageNames.Num() == 1 )
+		if (Identifiers.Num() == 1 )
 		{
-			const FString PackageNameStr = PackageNames[0].ToString();
+			const FString PackageNameStr = Identifiers[0].PackageName.ToString();
 			if ( FPackageName::IsValidLongPackageName(PackageNameStr, true) )
 			{
 				if ( PackageNameStr.StartsWith(TEXT("/Script")) )
@@ -168,4 +227,15 @@ bool UEdGraphNode_Reference::UsesThumbnail() const
 {
 	return bUsesThumbnail;
 }
+
+bool UEdGraphNode_Reference::IsPackage() const
+{
+	return bIsPackage;
+}
+
+bool UEdGraphNode_Reference::IsCollapsed() const
+{
+	return bIsCollapsed;
+}
+
 #undef LOCTEXT_NAMESPACE
\ No newline at end of file
diff --git a/Engine/Source/Editor/ReferenceViewer/Private/EdGraph_ReferenceViewer.cpp b/Engine/Source/Editor/ReferenceViewer/Private/EdGraph_ReferenceViewer.cpp
--- a/Engine/Source/Editor/ReferenceViewer/Private/EdGraph_ReferenceViewer.cpp
+++ b/Engine/Source/Editor/ReferenceViewer/Private/EdGraph_ReferenceViewer.cpp
@@ -15,9 +15,9 @@ UEdGraph_ReferenceViewer::UEdGraph_ReferenceViewer(const FObjectInitializer& Obj
 	bLimitSearchDepth = true;
 	bLimitSearchBreadth = true;
 	bIsShowSoftReferences = true;
-	bIsShowSoftDependencies = true;
 	bIsShowHardReferences = true;
-	bIsShowHardDependencies = true;
+	bIsShowSearchableNames = false;
+	bIsShowNativePackages = false;
 }
 
 void UEdGraph_ReferenceViewer::BeginDestroy()
@@ -31,21 +31,30 @@ void UEdGraph_ReferenceViewer::BeginDestroy()
 	Super::BeginDestroy();
 }
 
-void UEdGraph_ReferenceViewer::SetGraphRoot(const TArray<FName>& GraphRootPackageNames, const FIntPoint& GraphRootOrigin)
+void UEdGraph_ReferenceViewer::SetGraphRoot(const TArray<FAssetIdentifier>& GraphRootIdentifiers, const FIntPoint& GraphRootOrigin)
 {
-	CurrentGraphRootPackageNames = GraphRootPackageNames;
+	CurrentGraphRootIdentifiers = GraphRootIdentifiers;
 	CurrentGraphRootOrigin = GraphRootOrigin;
+
+	// If we're focused on a searchable name, enable that flag
+	for (const FAssetIdentifier& AssetId : GraphRootIdentifiers)
+	{
+		if (AssetId.IsValue())
+		{
+			bIsShowSearchableNames = true;
+		}
+	}
 }
 
-const TArray<FName>& UEdGraph_ReferenceViewer::GetCurrentGraphRootPackageNames() const
+const TArray<FAssetIdentifier>& UEdGraph_ReferenceViewer::GetCurrentGraphRootIdentifiers() const
 {
-	return CurrentGraphRootPackageNames;
+	return CurrentGraphRootIdentifiers;
 }
 
 UEdGraphNode_Reference* UEdGraph_ReferenceViewer::RebuildGraph()
 {
 	RemoveAllNodes();
-	UEdGraphNode_Reference* NewRootNode = ConstructNodes(CurrentGraphRootPackageNames, CurrentGraphRootOrigin);
+	UEdGraphNode_Reference* NewRootNode = ConstructNodes(CurrentGraphRootIdentifiers, CurrentGraphRootOrigin);
 	NotifyGraphChanged();
 
 	return NewRootNode;
@@ -66,19 +75,19 @@ bool UEdGraph_ReferenceViewer::IsShowSoftReferences() const
 	return bIsShowSoftReferences;
 }
 
-bool UEdGraph_ReferenceViewer::IsShowSoftDependencies() const
+bool UEdGraph_ReferenceViewer::IsShowHardReferences() const
 {
-	return bIsShowSoftDependencies;
+	return bIsShowHardReferences;
 }
 
-bool UEdGraph_ReferenceViewer::IsShowHardReferences() const
+bool UEdGraph_ReferenceViewer::IsShowSearchableNames() const
 {
-	return bIsShowHardReferences;
+	return bIsShowSearchableNames;
 }
 
-bool UEdGraph_ReferenceViewer::IsShowHardDependencies() const
+bool UEdGraph_ReferenceViewer::IsShowNativePackages() const
 {
-	return bIsShowHardDependencies;
+	return bIsShowNativePackages;
 }
 
 void UEdGraph_ReferenceViewer::SetSearchDepthLimitEnabled(bool newEnabled)
@@ -96,19 +105,19 @@ void UEdGraph_ReferenceViewer::SetShowSoftReferencesEnabled(bool newEnabled)
 	bIsShowSoftReferences = newEnabled;
 }
 
-void UEdGraph_ReferenceViewer::SetShowSoftDependenciesEnabled(bool newEnabled)
+void UEdGraph_ReferenceViewer::SetShowHardReferencesEnabled(bool newEnabled)
 {
-	bIsShowSoftDependencies = newEnabled;
+	bIsShowHardReferences = newEnabled;
 }
 
-void UEdGraph_ReferenceViewer::SetShowHardReferencesEnabled(bool newEnabled)
+void UEdGraph_ReferenceViewer::SetShowSearchableNames(bool newEnabled)
 {
-	bIsShowHardReferences = newEnabled;
+	bIsShowSearchableNames = newEnabled;
 }
 
-void UEdGraph_ReferenceViewer::SetShowHardDependenciesEnabled(bool newEnabled)
+void UEdGraph_ReferenceViewer::SetShowNativePackages(bool newEnabled)
 {
-	bIsShowHardDependencies = newEnabled;
+	bIsShowNativePackages = newEnabled;
 }
 
 int32 UEdGraph_ReferenceViewer::GetSearchDepthLimit() const
@@ -131,67 +140,113 @@ void UEdGraph_ReferenceViewer::SetSearchBreadthLimit(int32 NewBreadthLimit)
 	MaxSearchBreadth = NewBreadthLimit;
 }
 
-UEdGraphNode_Reference* UEdGraph_ReferenceViewer::ConstructNodes(const TArray<FName>& GraphRootPackageNames, const FIntPoint& GraphRootOrigin )
+EAssetRegistryDependencyType::Type UEdGraph_ReferenceViewer::GetReferenceSearchFlags(bool bReferencers) const
+{
+	int32 ReferenceFlags = 0;
+
+	if (bIsShowSoftReferences)
+	{
+		ReferenceFlags |= EAssetRegistryDependencyType::Soft;
+	}
+	if (bIsShowHardReferences)
+	{
+		ReferenceFlags |= EAssetRegistryDependencyType::Hard;
+	}
+	if (bIsShowSearchableNames)
+	{
+		ReferenceFlags |= EAssetRegistryDependencyType::SearchableName;
+	}
+
+	return (EAssetRegistryDependencyType::Type)ReferenceFlags;
+}
+
+UEdGraphNode_Reference* UEdGraph_ReferenceViewer::ConstructNodes(const TArray<FAssetIdentifier>& GraphRootIdentifiers, const FIntPoint& GraphRootOrigin )
 {
 	UEdGraphNode_Reference* RootNode = NULL;
 
-	if ( GraphRootPackageNames.Num() > 0 )
+	if (GraphRootIdentifiers.Num() > 0 )
 	{
-		TMap<FName, int32> ReferencerNodeSizes;
-		TSet<FName> VisitedReferencerSizeNames;
+		TMap<FAssetIdentifier, int32> ReferencerNodeSizes;
+		TSet<FAssetIdentifier> VisitedReferencerSizeNames;
 		int32 ReferencerDepth = 1;
-		RecursivelyGatherSizes(/*bReferencers=*/true, GraphRootPackageNames, ReferencerDepth, VisitedReferencerSizeNames, ReferencerNodeSizes);
+		RecursivelyGatherSizes(/*bReferencers=*/true, GraphRootIdentifiers, ReferencerDepth, VisitedReferencerSizeNames, ReferencerNodeSizes);
 
-		TMap<FName, int32> DependencyNodeSizes;
-		TSet<FName> VisitedDependencySizeNames;
+		TMap<FAssetIdentifier, int32> DependencyNodeSizes;
+		TSet<FAssetIdentifier> VisitedDependencySizeNames;
 		int32 DependencyDepth = 1;
-		RecursivelyGatherSizes(/*bReferencers=*/false, GraphRootPackageNames, DependencyDepth, VisitedDependencySizeNames, DependencyNodeSizes);
+		RecursivelyGatherSizes(/*bReferencers=*/false, GraphRootIdentifiers, DependencyDepth, VisitedDependencySizeNames, DependencyNodeSizes);
+
+		TSet<FName> AllPackageNames;
+
+		auto AddPackage = [](const FAssetIdentifier& AssetId, TSet<FName>& PackageNames)
+		{ 
+			// Only look for asset data if this is a package
+			if (!AssetId.IsValue())
+			{
+				PackageNames.Add(AssetId.PackageName);
+			}
+		};
+
+		for (const FAssetIdentifier& AssetId : VisitedReferencerSizeNames)
+		{
+			AddPackage(AssetId, AllPackageNames);
+		}
+
+		for (const FAssetIdentifier& AssetId : VisitedDependencySizeNames)
+		{
+			AddPackage(AssetId, AllPackageNames);
+		}
 
-		TSet<FName> AllPackageNames = VisitedReferencerSizeNames;
-		AllPackageNames.Append(VisitedDependencySizeNames);
 		TMap<FName, FAssetData> PackagesToAssetDataMap;
 		GatherAssetData(AllPackageNames, PackagesToAssetDataMap);
 
 		// Create the root node
 		RootNode = CreateReferenceNode();
-		RootNode->SetupReferenceNode(GraphRootOrigin, GraphRootPackageNames, PackagesToAssetDataMap.FindRef(GraphRootPackageNames[0]));
+		RootNode->SetupReferenceNode(GraphRootOrigin, GraphRootIdentifiers, PackagesToAssetDataMap.FindRef(GraphRootIdentifiers[0].PackageName));
 
-		TSet<FName> VisitedReferencerNames;
+		TSet<FAssetIdentifier> VisitedReferencerNames;
 		int32 VisitedReferencerDepth = 1;
-		RecursivelyConstructNodes(/*bReferencers=*/true, RootNode, GraphRootPackageNames, GraphRootOrigin, ReferencerNodeSizes, PackagesToAssetDataMap, VisitedReferencerDepth, VisitedReferencerNames);
+		RecursivelyConstructNodes(/*bReferencers=*/true, RootNode, GraphRootIdentifiers, GraphRootOrigin, ReferencerNodeSizes, PackagesToAssetDataMap, VisitedReferencerDepth, VisitedReferencerNames);
 
-		TSet<FName> VisitedDependencyNames;
+		TSet<FAssetIdentifier> VisitedDependencyNames;
 		int32 VisitedDependencyDepth = 1;
-		RecursivelyConstructNodes(/*bReferencers=*/false, RootNode, GraphRootPackageNames, GraphRootOrigin, DependencyNodeSizes, PackagesToAssetDataMap, VisitedDependencyDepth, VisitedDependencyNames);
+		RecursivelyConstructNodes(/*bReferencers=*/false, RootNode, GraphRootIdentifiers, GraphRootOrigin, DependencyNodeSizes, PackagesToAssetDataMap, VisitedDependencyDepth, VisitedDependencyNames);
 	}
 
 	return RootNode;
 }
 
-int32 UEdGraph_ReferenceViewer::RecursivelyGatherSizes(bool bReferencers, const TArray<FName>& PackageNames, int32 CurrentDepth, TSet<FName>& VisitedNames, TMap<FName, int32>& OutNodeSizes) const
+int32 UEdGraph_ReferenceViewer::RecursivelyGatherSizes(bool bReferencers, const TArray<FAssetIdentifier>& Identifiers, int32 CurrentDepth, TSet<FAssetIdentifier>& VisitedNames, TMap<FAssetIdentifier, int32>& OutNodeSizes) const
 {
-	check(PackageNames.Num() > 0);
+	check(Identifiers.Num() > 0);
 
-	VisitedNames.Append(PackageNames);
+	VisitedNames.Append(Identifiers);
 
 	FAssetRegistryModule& AssetRegistryModule = FModuleManager::LoadModuleChecked<FAssetRegistryModule>("AssetRegistry");
-	TArray<FName> ReferenceNames;
+	TArray<FAssetIdentifier> ReferenceNames;
 
 	if ( bReferencers )
 	{
-		for ( auto NameIt = PackageNames.CreateConstIterator(); NameIt; ++NameIt )
+		for (const FAssetIdentifier& AssetId : Identifiers)
 		{
-			AssetRegistryModule.Get().GetReferencers(*NameIt, ReferenceNames);
+			AssetRegistryModule.Get().GetReferencers(AssetId, ReferenceNames, GetReferenceSearchFlags(bReferencers));
 		}
 	}
 	else
 	{
-		for ( auto NameIt = PackageNames.CreateConstIterator(); NameIt; ++NameIt )
+		for (const FAssetIdentifier& AssetId : Identifiers)
 		{
-			AssetRegistryModule.Get().GetDependencies(*NameIt, ReferenceNames);
+			AssetRegistryModule.Get().GetDependencies(AssetId, ReferenceNames, GetReferenceSearchFlags(bReferencers));
 		}
 	}
 
+	if (!bIsShowNativePackages)
+	{
+		auto RemoveNativePackage = [](const FAssetIdentifier& InAsset) { return InAsset.PackageName.ToString().StartsWith(TEXT("/Script")) && !InAsset.IsValue(); };
+
+		ReferenceNames.RemoveAll(RemoveNativePackage);
+	}
+
 	int32 NodeSize = 0;
 	if ( ReferenceNames.Num() > 0 && !ExceedsMaxSearchDepth(CurrentDepth) )
 	{
@@ -200,14 +255,14 @@ int32 UEdGraph_ReferenceViewer::RecursivelyGatherSizes(bool bReferencers, const
 
 		// Since there are referencers, use the size of all your combined referencers.
 		// Do not count your own size since there could just be a horizontal line of nodes
-		for ( auto RefIt = ReferenceNames.CreateConstIterator(); RefIt; ++RefIt )
+		for (const FAssetIdentifier& AssetId : ReferenceNames)
 		{
-			if ( !VisitedNames.Contains(*RefIt) )
+			if ( !VisitedNames.Contains(AssetId) )
 			{
 				if ( !ExceedsMaxSearchBreadth(NumReferencesMade) )
 				{
-					TArray<FName> NewPackageNames;
-					NewPackageNames.Add(*RefIt);
+					TArray<FAssetIdentifier> NewPackageNames;
+					NewPackageNames.Add(AssetId);
 					NodeSize += RecursivelyGatherSizes(bReferencers, NewPackageNames, CurrentDepth + 1, VisitedNames, OutNodeSizes);
 					NumReferencesMade++;
 				}
@@ -231,7 +286,7 @@ int32 UEdGraph_ReferenceViewer::RecursivelyGatherSizes(bool bReferencers, const
 		NodeSize = 1;
 	}
 
-	OutNodeSizes.Add(PackageNames[0], NodeSize);
+	OutNodeSizes.Add(Identifiers[0], NodeSize);
 	return NodeSize;
 }
 
@@ -256,14 +311,14 @@ void UEdGraph_ReferenceViewer::GatherAssetData(const TSet<FName>& AllPackageName
 	}
 }
 
-UEdGraphNode_Reference* UEdGraph_ReferenceViewer::RecursivelyConstructNodes(bool bReferencers, UEdGraphNode_Reference* RootNode, const TArray<FName>& PackageNames, const FIntPoint& NodeLoc, const TMap<FName, int32>& NodeSizes, const TMap<FName, FAssetData>& PackagesToAssetDataMap, int32 CurrentDepth, TSet<FName>& VisitedNames)
+UEdGraphNode_Reference* UEdGraph_ReferenceViewer::RecursivelyConstructNodes(bool bReferencers, UEdGraphNode_Reference* RootNode, const TArray<FAssetIdentifier>& Identifiers, const FIntPoint& NodeLoc, const TMap<FAssetIdentifier, int32>& NodeSizes, const TMap<FName, FAssetData>& PackagesToAssetDataMap, int32 CurrentDepth, TSet<FAssetIdentifier>& VisitedNames)
 {
-	check(PackageNames.Num() > 0);
+	check(Identifiers.Num() > 0);
 
-	VisitedNames.Append(PackageNames);
+	VisitedNames.Append(Identifiers);
 
 	UEdGraphNode_Reference* NewNode = NULL;
-	if ( RootNode->GetPackageName() == PackageNames[0] )
+	if ( RootNode->GetIdentifier() == Identifiers[0] )
 	{
 		// Don't create the root node. It is already created!
 		NewNode = RootNode;
@@ -271,51 +326,39 @@ UEdGraphNode_Reference* UEdGraph_ReferenceViewer::RecursivelyConstructNodes(bool
 	else
 	{
 		NewNode = CreateReferenceNode();
-		NewNode->SetupReferenceNode(NodeLoc, PackageNames, PackagesToAssetDataMap.FindRef(PackageNames[0]));
+		NewNode->SetupReferenceNode(NodeLoc, Identifiers, PackagesToAssetDataMap.FindRef(Identifiers[0].PackageName));
 	}
 
 	FAssetRegistryModule& AssetRegistryModule = FModuleManager::LoadModuleChecked<FAssetRegistryModule>("AssetRegistry");
-	TArray<FName> ReferenceNames;
-	TArray<FName> HardReferenceNames;
+	TArray<FAssetIdentifier> ReferenceNames;
+	TArray<FAssetIdentifier> HardReferenceNames;
 	if ( bReferencers )
 	{
-		for ( auto NameIt = PackageNames.CreateConstIterator(); NameIt; ++NameIt )
+		for (const FAssetIdentifier& AssetId : Identifiers)
 		{
-			AssetRegistryModule.Get().GetReferencers(*NameIt, HardReferenceNames, EAssetRegistryDependencyType::Hard);
-			if (bIsShowSoftReferences && bIsShowHardReferences)
-			{
-				AssetRegistryModule.Get().GetReferencers(*NameIt, ReferenceNames, EAssetRegistryDependencyType::All);
-			}
-			else if (bIsShowSoftReferences && !bIsShowHardReferences)
-			{
-				AssetRegistryModule.Get().GetReferencers(*NameIt, ReferenceNames, EAssetRegistryDependencyType::Soft);
-			}
-			else
-			{
-				AssetRegistryModule.Get().GetReferencers(*NameIt, ReferenceNames, EAssetRegistryDependencyType::Hard);
-			}
+			AssetRegistryModule.Get().GetReferencers(AssetId, HardReferenceNames, EAssetRegistryDependencyType::Hard);
+
+			AssetRegistryModule.Get().GetReferencers(AssetId, ReferenceNames, GetReferenceSearchFlags(bReferencers));
 		}
 	}
 	else
 	{
-		for ( auto NameIt = PackageNames.CreateConstIterator(); NameIt; ++NameIt )
+		for (const FAssetIdentifier& AssetId : Identifiers)
 		{
-			AssetRegistryModule.Get().GetDependencies(*NameIt, HardReferenceNames, EAssetRegistryDependencyType::Hard);
-			if (bIsShowSoftDependencies && bIsShowHardDependencies)
-			{
-				AssetRegistryModule.Get().GetDependencies(*NameIt, ReferenceNames, EAssetRegistryDependencyType::All);
-			}
-			else if (bIsShowSoftDependencies && !bIsShowHardDependencies)
-			{
-				AssetRegistryModule.Get().GetDependencies(*NameIt, ReferenceNames, EAssetRegistryDependencyType::Soft);
-			}
-			else
-			{
-				AssetRegistryModule.Get().GetDependencies(*NameIt, ReferenceNames, EAssetRegistryDependencyType::Hard);
-			}
+			AssetRegistryModule.Get().GetDependencies(AssetId, HardReferenceNames, EAssetRegistryDependencyType::Hard);
+
+			AssetRegistryModule.Get().GetDependencies(AssetId, ReferenceNames, GetReferenceSearchFlags(bReferencers));
 		}
 	}
 
+	if (!bIsShowNativePackages)
+	{
+		auto RemoveNativePackage = [](const FAssetIdentifier& InAsset) { return InAsset.PackageName.ToString().StartsWith(TEXT("/Script")) && !InAsset.IsValue(); };
+
+		HardReferenceNames.RemoveAll(RemoveNativePackage);
+		ReferenceNames.RemoveAll(RemoveNativePackage);
+	}
+
 	if ( ReferenceNames.Num() > 0 && !ExceedsMaxSearchDepth(CurrentDepth) )
 	{
 		FIntPoint ReferenceNodeLoc = NodeLoc;
@@ -332,7 +375,7 @@ UEdGraphNode_Reference* UEdGraph_ReferenceViewer::RecursivelyConstructNodes(bool
 		}
 
 		const int32 NodeSizeY = 200;
-		const int32 TotalReferenceSizeY = NodeSizes.FindChecked(PackageNames[0]) * NodeSizeY;
+		const int32 TotalReferenceSizeY = NodeSizes.FindChecked(Identifiers[0]) * NodeSizeY;
 
 		ReferenceNodeLoc.Y -= TotalReferenceSizeY * 0.5f;
 		ReferenceNodeLoc.Y += NodeSizeY * 0.5f;
@@ -341,34 +384,35 @@ UEdGraphNode_Reference* UEdGraph_ReferenceViewer::RecursivelyConstructNodes(bool
 		int32 NumReferencesExceedingMax = 0;
 		for ( int32 RefIdx = 0; RefIdx < ReferenceNames.Num(); ++RefIdx )
 		{
-			FName ReferenceName = ReferenceNames[RefIdx];
+			FAssetIdentifier ReferenceName = ReferenceNames[RefIdx];
 
 			if ( !VisitedNames.Contains(ReferenceName) )
 			{
 				if ( !ExceedsMaxSearchBreadth(NumReferencesMade) )
 				{
+					int32 ThisNodeSizeY = ReferenceName.IsValue() ? 100 : NodeSizeY;
+
 					const int32 RefSizeY = NodeSizes.FindChecked(ReferenceName);
 					FIntPoint RefNodeLoc;
 					RefNodeLoc.X = ReferenceNodeLoc.X;
-					RefNodeLoc.Y = ReferenceNodeLoc.Y + RefSizeY * NodeSizeY * 0.5 - NodeSizeY * 0.5;
-					TArray<FName> NewPackageNames;
-					NewPackageNames.Add(ReferenceName);
-					UEdGraphNode_Reference* ReferenceNode = RecursivelyConstructNodes(bReferencers, RootNode, NewPackageNames, RefNodeLoc, NodeSizes, PackagesToAssetDataMap, CurrentDepth + 1, VisitedNames);
-					for (FName NodeName : HardReferenceNames)
+					RefNodeLoc.Y = ReferenceNodeLoc.Y + RefSizeY * ThisNodeSizeY * 0.5 - ThisNodeSizeY * 0.5;
+					
+					TArray<FAssetIdentifier> NewIdentifiers;
+					NewIdentifiers.Add(ReferenceName);
+					
+					UEdGraphNode_Reference* ReferenceNode = RecursivelyConstructNodes(bReferencers, RootNode, NewIdentifiers, RefNodeLoc, NodeSizes, PackagesToAssetDataMap, CurrentDepth + 1, VisitedNames);
+					if (HardReferenceNames.Contains(ReferenceName))
 					{
-						FName NewNodeFName = FName(*ReferenceNode->NodeComment);
-						if (NodeName == FName(*ReferenceNode->NodeComment))
+						if (bReferencers)
 						{
-							if (bReferencers)
-							{
-								ReferenceNode->GetDependencyPin()->PinType.PinCategory = TEXT("hard");
-							}
-							else
-							{
-								ReferenceNode->GetReferencerPin()->PinType.PinCategory = TEXT("hard");
-							}
+							ReferenceNode->GetDependencyPin()->PinType.PinCategory = TEXT("hard");
+						}
+						else
+						{
+							ReferenceNode->GetReferencerPin()->PinType.PinCategory = TEXT("hard");
 						}
 					}
+
 					if ( ensure(ReferenceNode) )
 					{
 						if ( bReferencers )
@@ -380,7 +424,7 @@ UEdGraphNode_Reference* UEdGraph_ReferenceViewer::RecursivelyConstructNodes(bool
 							ReferenceNode->AddReferencer( NewNode );
 						}
 
-						ReferenceNodeLoc.Y += RefSizeY * NodeSizeY;
+						ReferenceNodeLoc.Y += RefSizeY * ThisNodeSizeY;
 					}
 
 					NumReferencesMade++;
diff --git a/Engine/Source/Editor/ReferenceViewer/Private/HistoryManager.h b/Engine/Source/Editor/ReferenceViewer/Private/HistoryManager.h
--- a/Engine/Source/Editor/ReferenceViewer/Private/HistoryManager.h
+++ b/Engine/Source/Editor/ReferenceViewer/Private/HistoryManager.h
@@ -9,7 +9,7 @@ struct FReferenceViewerHistoryData
 	FText HistoryDesc;
 
 	/** The list of package names to serve as the root */
-	TArray<FName> PackageNames;
+	TArray<FAssetIdentifier> Identifiers;
 };
 
 /** The delegate for when history data should be applied */
diff --git a/Engine/Source/Editor/ReferenceViewer/Private/ReferenceViewer.cpp b/Engine/Source/Editor/ReferenceViewer/Private/ReferenceViewer.cpp
--- a/Engine/Source/Editor/ReferenceViewer/Private/ReferenceViewer.cpp
+++ b/Engine/Source/Editor/ReferenceViewer/Private/ReferenceViewer.cpp
@@ -52,11 +52,11 @@ public:
 		FGlobalTabmanager::Get()->UnregisterNomadTabSpawner(ReferenceViewerTabId);
 	}
 
-	virtual void InvokeReferenceViewerTab(const TArray<FName>& GraphRootPackageNames) override
+	virtual void InvokeReferenceViewerTab(const TArray<FAssetIdentifier>& GraphRootIdentifiers) override
 	{
 		TSharedRef<SDockTab> NewTab = FGlobalTabmanager::Get()->InvokeTab( ReferenceViewerTabId );
 		TSharedRef<SReferenceViewer> ReferenceViewer = StaticCastSharedRef<SReferenceViewer>( NewTab->GetContent() );
-		ReferenceViewer->SetGraphRootPackageNames(GraphRootPackageNames);
+		ReferenceViewer->SetGraphRootPackageNames(GraphRootIdentifiers);
 	}
 
 private:
diff --git a/Engine/Source/Editor/ReferenceViewer/Private/SReferenceNode.cpp b/Engine/Source/Editor/ReferenceViewer/Private/SReferenceNode.cpp
--- a/Engine/Source/Editor/ReferenceViewer/Private/SReferenceNode.cpp
+++ b/Engine/Source/Editor/ReferenceViewer/Private/SReferenceNode.cpp
@@ -12,15 +12,14 @@
 void SReferenceNode::Construct( const FArguments& InArgs, UEdGraphNode_Reference* InNode )
 {
 	const int32 ThumbnailSize = 128;
-	bUsesThumbnail = InNode->UsesThumbnail();
 
-	if (bUsesThumbnail)
+	if (InNode->UsesThumbnail())
 	{
 		// Create a thumbnail from the graph's thumbnail pool
 		TSharedPtr<FAssetThumbnailPool> AssetThumbnailPool = InNode->GetReferenceViewerGraph()->GetAssetThumbnailPool();
 		AssetThumbnail = MakeShareable( new FAssetThumbnail( InNode->GetAssetData(), ThumbnailSize, ThumbnailSize, AssetThumbnailPool ) );
 	}
-	else
+	else if (InNode->IsPackage() || InNode->IsCollapsed())
 	{
 		// Just make a generic thumbnail
 		AssetThumbnail = MakeShareable( new FAssetThumbnail( InNode->GetAssetData(), ThumbnailSize, ThumbnailSize, NULL ) );
