Copying //UE4/Dev-Framework to //UE4/Dev-Main (Source: //UE4/Dev-Framework @ 3182037)

diff --git a/Engine/Source/Editor/ReferenceViewer/Private/SReferenceViewer.cpp b/Engine/Source/Editor/ReferenceViewer/Private/SReferenceViewer.cpp
--- a/Engine/Source/Editor/ReferenceViewer/Private/SReferenceViewer.cpp
+++ b/Engine/Source/Editor/ReferenceViewer/Private/SReferenceViewer.cpp
@@ -249,6 +249,7 @@ void SReferenceViewer::Construct( const FArguments& InArgs )
 						]
 					]
 
+				
 					+ SVerticalBox::Slot()
 					.AutoHeight()
 					[
@@ -259,7 +260,7 @@ void SReferenceViewer::Construct( const FArguments& InArgs )
 						.Padding(2.f)
 						[
 							SNew(STextBlock)
-							.Text(LOCTEXT("ShowHideSoftDependencies", "Show Soft Dependencies"))
+							.Text(LOCTEXT("ShowHideHardReferences", "Show Hard References"))
 						]
 
 						+SHorizontalBox::Slot()
@@ -268,8 +269,8 @@ void SReferenceViewer::Construct( const FArguments& InArgs )
 						.Padding(2.f)
 						[
 							SNew(SCheckBox)
-							.OnCheckStateChanged(this, &SReferenceViewer::OnShowSoftDependenciesChanged)
-							.IsChecked(this, &SReferenceViewer::IsShowSoftDependenciesChecked)
+							.OnCheckStateChanged(this, &SReferenceViewer::OnShowHardReferencesChanged)
+							.IsChecked(this, &SReferenceViewer::IsShowHardReferencesChecked)
 						]
 					]
 
@@ -283,7 +284,7 @@ void SReferenceViewer::Construct( const FArguments& InArgs )
 						.Padding(2.f)
 						[
 							SNew(STextBlock)
-							.Text(LOCTEXT("ShowHideHardReferences", "Show Hard References"))
+							.Text(LOCTEXT("ShowHideSearchableNames", "Show Searchable Names"))
 						]
 
 						+SHorizontalBox::Slot()
@@ -292,8 +293,8 @@ void SReferenceViewer::Construct( const FArguments& InArgs )
 						.Padding(2.f)
 						[
 							SNew(SCheckBox)
-							.OnCheckStateChanged(this, &SReferenceViewer::OnShowHardReferencesChanged)
-							.IsChecked(this, &SReferenceViewer::IsShowHardReferencesChecked)
+							.OnCheckStateChanged(this, &SReferenceViewer::OnShowSearchableNamesChanged)
+							.IsChecked(this, &SReferenceViewer::IsShowSearchableNamesChecked)
 						]
 					]
 
@@ -307,7 +308,7 @@ void SReferenceViewer::Construct( const FArguments& InArgs )
 						.Padding(2.f)
 						[
 							SNew(STextBlock)
-							.Text(LOCTEXT("ShowHideHardDependencies", "Show Hard Dependencies"))
+							.Text(LOCTEXT("ShowHideNativePackages", "Show Native Packages"))
 						]
 
 						+SHorizontalBox::Slot()
@@ -316,8 +317,8 @@ void SReferenceViewer::Construct( const FArguments& InArgs )
 						.Padding(2.f)
 						[
 							SNew(SCheckBox)
-							.OnCheckStateChanged(this, &SReferenceViewer::OnShowHardDependenciesChanged)
-							.IsChecked(this, &SReferenceViewer::IsShowHardDependenciesChecked)
+							.OnCheckStateChanged(this, &SReferenceViewer::OnShowNativePackagesChanged)
+							.IsChecked(this, &SReferenceViewer::IsShowNativePackagesChecked)
 						]
 					]
 				]
@@ -334,9 +335,9 @@ void SReferenceViewer::Construct( const FArguments& InArgs )
 	];
 }
 
-void SReferenceViewer::SetGraphRootPackageNames(const TArray<FName>& NewGraphRootPackageNames)
+void SReferenceViewer::SetGraphRootPackageNames(const TArray<FAssetIdentifier>& NewGraphRootIdentifiers)
 {
-	GraphObj->SetGraphRoot(NewGraphRootPackageNames);
+	GraphObj->SetGraphRoot(NewGraphRootIdentifiers);
 	
 	FAssetRegistryModule& AssetRegistryModule = FModuleManager::LoadModuleChecked<FAssetRegistryModule>(TEXT("AssetRegistry"));
 	if ( AssetRegistryModule.Get().IsLoadingAssets() )
@@ -428,14 +429,14 @@ FText SReferenceViewer::GetAddressBarText() const
 	{
 		if (TemporaryPathBeingEdited.IsEmpty())
 		{
-			const TArray<FName>& CurrentGraphRootPackageNames = GraphObj->GetCurrentGraphRootPackageNames();
+			const TArray<FAssetIdentifier>& CurrentGraphRootPackageNames = GraphObj->GetCurrentGraphRootIdentifiers();
 			if (CurrentGraphRootPackageNames.Num() == 1)
 			{
-				return FText::FromName(CurrentGraphRootPackageNames[0]);
+				return FText::FromString(CurrentGraphRootPackageNames[0].ToString());
 			}
 			else if (CurrentGraphRootPackageNames.Num() > 1)
 			{
-				return FText::Format(LOCTEXT("AddressBarMultiplePackagesText", "{0} and {1} others"), FText::FromName(CurrentGraphRootPackageNames[0]), FText::AsNumber(CurrentGraphRootPackageNames.Num()));
+				return FText::Format(LOCTEXT("AddressBarMultiplePackagesText", "{0} and {1} others"), FText::FromString(CurrentGraphRootPackageNames[0].ToString()), FText::AsNumber(CurrentGraphRootPackageNames.Num()));
 			}
 		}
 		else
@@ -451,10 +452,8 @@ void SReferenceViewer::OnAddressBarTextCommitted(const FText& NewText, ETextComm
 {
 	if (CommitInfo == ETextCommit::OnEnter)
 	{
-		const FName NewPath(*NewText.ToString());
-
-		TArray<FName> NewPaths;
-		NewPaths.Add(NewPath);
+		TArray<FAssetIdentifier> NewPaths;
+		NewPaths.Add(FAssetIdentifier::FromString(NewText.ToString()));
 
 		SetGraphRootPackageNames(NewPaths);
 	}
@@ -471,7 +470,7 @@ void SReferenceViewer::OnApplyHistoryData(const FReferenceViewerHistoryData& His
 {
 	if ( GraphObj )
 	{
-		GraphObj->SetGraphRoot(History.PackageNames);
+		GraphObj->SetGraphRoot(History.Identifiers);
 		UEdGraphNode_Reference* NewRootNode = GraphObj->RebuildGraph();
 		
 		if ( NewRootNode && ensure(GraphEditorPtr.IsValid()) )
@@ -485,26 +484,14 @@ void SReferenceViewer::OnUpdateHistoryData(FReferenceViewerHistoryData& HistoryD
 {
 	if ( GraphObj )
 	{
-		const TArray<FName>& CurrentGraphRootPackageNames = GraphObj->GetCurrentGraphRootPackageNames();
-		if ( CurrentGraphRootPackageNames.Num() == 1 )
-		{
-			HistoryData.HistoryDesc = FText::FromName(CurrentGraphRootPackageNames[0]);
-		}
-		else if ( CurrentGraphRootPackageNames.Num() > 1 )
-		{
-			HistoryData.HistoryDesc = FText::Format(LOCTEXT("AddressBarMultiplePackagesText", "{0} and {1} others"), FText::FromName(CurrentGraphRootPackageNames[0]), FText::AsNumber(CurrentGraphRootPackageNames.Num()));
-		}
-		else
-		{
-			HistoryData.HistoryDesc = FText::GetEmpty();
-		}
-
-		HistoryData.PackageNames = CurrentGraphRootPackageNames;
+		const TArray<FAssetIdentifier>& CurrentGraphRootIdentifiers = GraphObj->GetCurrentGraphRootIdentifiers();
+		HistoryData.HistoryDesc = GetAddressBarText();
+		HistoryData.Identifiers = CurrentGraphRootIdentifiers;
 	}
 	else
 	{
 		HistoryData.HistoryDesc = FText::GetEmpty();
-		HistoryData.PackageNames.Empty();
+		HistoryData.Identifiers.Empty();
 	}
 }
 
@@ -592,20 +579,21 @@ ECheckBoxState SReferenceViewer::IsShowSoftReferencesChecked() const
 	}
 }
 
-void SReferenceViewer::OnShowSoftDependenciesChanged(ECheckBoxState NewState)
+void SReferenceViewer::OnShowHardReferencesChanged(ECheckBoxState NewState)
 {
 	if (GraphObj)
 	{
-		GraphObj->SetShowSoftDependenciesEnabled(NewState == ECheckBoxState::Checked);
+		GraphObj->SetShowHardReferencesEnabled(NewState == ECheckBoxState::Checked);
 		GraphObj->RebuildGraph();
 	}
 }
 
-ECheckBoxState SReferenceViewer::IsShowSoftDependenciesChecked() const
+
+ECheckBoxState SReferenceViewer::IsShowHardReferencesChecked() const
 {
 	if (GraphObj)
 	{
-		return GraphObj->IsShowSoftDependencies() ? ECheckBoxState::Checked : ECheckBoxState::Unchecked;
+		return GraphObj->IsShowHardReferences() ? ECheckBoxState::Checked : ECheckBoxState::Unchecked;
 	}
 	else
 	{
@@ -613,21 +601,20 @@ ECheckBoxState SReferenceViewer::IsShowSoftDependenciesChecked() const
 	}
 }
 
-void SReferenceViewer::OnShowHardReferencesChanged(ECheckBoxState NewState)
+void SReferenceViewer::OnShowSearchableNamesChanged(ECheckBoxState NewState)
 {
 	if (GraphObj)
 	{
-		GraphObj->SetShowHardReferencesEnabled(NewState == ECheckBoxState::Checked);
+		GraphObj->SetShowSearchableNames(NewState == ECheckBoxState::Checked);
 		GraphObj->RebuildGraph();
 	}
 }
 
-
-ECheckBoxState SReferenceViewer::IsShowHardReferencesChecked() const
+ECheckBoxState SReferenceViewer::IsShowSearchableNamesChecked() const
 {
 	if (GraphObj)
 	{
-		return GraphObj->IsShowHardReferences() ? ECheckBoxState::Checked : ECheckBoxState::Unchecked;
+		return GraphObj->IsShowSearchableNames() ? ECheckBoxState::Checked : ECheckBoxState::Unchecked;
 	}
 	else
 	{
@@ -635,20 +622,20 @@ ECheckBoxState SReferenceViewer::IsShowHardReferencesChecked() const
 	}
 }
 
-void SReferenceViewer::OnShowHardDependenciesChanged(ECheckBoxState NewState)
+void SReferenceViewer::OnShowNativePackagesChanged(ECheckBoxState NewState)
 {
 	if (GraphObj)
 	{
-		GraphObj->SetShowHardDependenciesEnabled(NewState == ECheckBoxState::Checked);
+		GraphObj->SetShowNativePackages(NewState == ECheckBoxState::Checked);
 		GraphObj->RebuildGraph();
 	}
 }
 
-ECheckBoxState SReferenceViewer::IsShowHardDependenciesChecked() const
+ECheckBoxState SReferenceViewer::IsShowNativePackagesChecked() const
 {
 	if (GraphObj)
 	{
-		return GraphObj->IsShowHardDependencies() ? ECheckBoxState::Checked : ECheckBoxState::Unchecked;
+		return GraphObj->IsShowNativePackages() ? ECheckBoxState::Checked : ECheckBoxState::Unchecked;
 	}
 	else
 	{
@@ -684,7 +671,8 @@ void SReferenceViewer::RegisterActions()
 
 	ReferenceViewerActions->MapAction(
 		FGlobalEditorCommonCommands::Get().FindInContentBrowser,
-		FExecuteAction::CreateSP(this, &SReferenceViewer::ShowSelectionInContentBrowser));
+		FExecuteAction::CreateSP(this, &SReferenceViewer::ShowSelectionInContentBrowser),
+		FCanExecuteAction::CreateSP(this, &SReferenceViewer::HasAtLeastOnePackageNodeSelected));
 
 	ReferenceViewerActions->MapAction(
 		FReferenceViewerActions::Get().OpenSelectedInAssetEditor,
@@ -699,52 +687,52 @@ void SReferenceViewer::RegisterActions()
 	ReferenceViewerActions->MapAction(
 		FReferenceViewerActions::Get().ListReferencedObjects,
 		FExecuteAction::CreateSP(this, &SReferenceViewer::ListReferencedObjects),
-		FCanExecuteAction::CreateSP(this, &SReferenceViewer::HasExactlyOneNodeSelected));
+		FCanExecuteAction::CreateSP(this, &SReferenceViewer::HasExactlyOnePackageNodeSelected));
 
 	ReferenceViewerActions->MapAction(
 		FReferenceViewerActions::Get().ListObjectsThatReference,
 		FExecuteAction::CreateSP(this, &SReferenceViewer::ListObjectsThatReference),
-		FCanExecuteAction::CreateSP(this, &SReferenceViewer::HasExactlyOneNodeSelected));
+		FCanExecuteAction::CreateSP(this, &SReferenceViewer::HasExactlyOnePackageNodeSelected));
 
 	ReferenceViewerActions->MapAction(
 		FReferenceViewerActions::Get().MakeLocalCollectionWithReferencers,
-		FExecuteAction::CreateSP(this, &SReferenceViewer::MakeCollectionWithReferenersOrDependencies, ECollectionShareType::CST_Local, true),
-		FCanExecuteAction::CreateSP(this, &SReferenceViewer::HasExactlyOneNodeSelected));
+		FExecuteAction::CreateSP(this, &SReferenceViewer::MakeCollectionWithReferencersOrDependencies, ECollectionShareType::CST_Local, true),
+		FCanExecuteAction::CreateSP(this, &SReferenceViewer::HasExactlyOnePackageNodeSelected));
 	
 	ReferenceViewerActions->MapAction(
 		FReferenceViewerActions::Get().MakePrivateCollectionWithReferencers,
-		FExecuteAction::CreateSP(this, &SReferenceViewer::MakeCollectionWithReferenersOrDependencies, ECollectionShareType::CST_Private, true),
-		FCanExecuteAction::CreateSP(this, &SReferenceViewer::HasExactlyOneNodeSelected));
+		FExecuteAction::CreateSP(this, &SReferenceViewer::MakeCollectionWithReferencersOrDependencies, ECollectionShareType::CST_Private, true),
+		FCanExecuteAction::CreateSP(this, &SReferenceViewer::HasExactlyOnePackageNodeSelected));
 	
 	ReferenceViewerActions->MapAction(
 		FReferenceViewerActions::Get().MakeSharedCollectionWithReferencers,
-		FExecuteAction::CreateSP(this, &SReferenceViewer::MakeCollectionWithReferenersOrDependencies, ECollectionShareType::CST_Shared, true),
-		FCanExecuteAction::CreateSP(this, &SReferenceViewer::HasExactlyOneNodeSelected));
+		FExecuteAction::CreateSP(this, &SReferenceViewer::MakeCollectionWithReferencersOrDependencies, ECollectionShareType::CST_Shared, true),
+		FCanExecuteAction::CreateSP(this, &SReferenceViewer::HasExactlyOnePackageNodeSelected));
 
 	ReferenceViewerActions->MapAction(
 		FReferenceViewerActions::Get().MakeLocalCollectionWithDependencies,
-		FExecuteAction::CreateSP(this, &SReferenceViewer::MakeCollectionWithReferenersOrDependencies, ECollectionShareType::CST_Local, false),
-		FCanExecuteAction::CreateSP(this, &SReferenceViewer::HasExactlyOneNodeSelected));
+		FExecuteAction::CreateSP(this, &SReferenceViewer::MakeCollectionWithReferencersOrDependencies, ECollectionShareType::CST_Local, false),
+		FCanExecuteAction::CreateSP(this, &SReferenceViewer::HasExactlyOnePackageNodeSelected));
 
 	ReferenceViewerActions->MapAction(
 		FReferenceViewerActions::Get().MakePrivateCollectionWithDependencies,
-		FExecuteAction::CreateSP(this, &SReferenceViewer::MakeCollectionWithReferenersOrDependencies, ECollectionShareType::CST_Private, false),
-		FCanExecuteAction::CreateSP(this, &SReferenceViewer::HasExactlyOneNodeSelected));
+		FExecuteAction::CreateSP(this, &SReferenceViewer::MakeCollectionWithReferencersOrDependencies, ECollectionShareType::CST_Private, false),
+		FCanExecuteAction::CreateSP(this, &SReferenceViewer::HasExactlyOnePackageNodeSelected));
 
 	ReferenceViewerActions->MapAction(
 		FReferenceViewerActions::Get().MakeSharedCollectionWithDependencies,
-		FExecuteAction::CreateSP(this, &SReferenceViewer::MakeCollectionWithReferenersOrDependencies, ECollectionShareType::CST_Shared, false),
-		FCanExecuteAction::CreateSP(this, &SReferenceViewer::HasExactlyOneNodeSelected));
+		FExecuteAction::CreateSP(this, &SReferenceViewer::MakeCollectionWithReferencersOrDependencies, ECollectionShareType::CST_Shared, false),
+		FCanExecuteAction::CreateSP(this, &SReferenceViewer::HasExactlyOnePackageNodeSelected));
 	
 	ReferenceViewerActions->MapAction(
 		FReferenceViewerActions::Get().ShowSizeMap,
 		FExecuteAction::CreateSP(this, &SReferenceViewer::ShowSizeMap),
-		FCanExecuteAction::CreateSP(this, &SReferenceViewer::HasAtLeastOneNodeSelected));
+		FCanExecuteAction::CreateSP(this, &SReferenceViewer::HasAtLeastOnePackageNodeSelected));
 
 	ReferenceViewerActions->MapAction(
 		FReferenceViewerActions::Get().ShowReferenceTree,
 		FExecuteAction::CreateSP(this, &SReferenceViewer::ShowReferenceTree),
-		FCanExecuteAction::CreateSP(this, &SReferenceViewer::HasExactlyOneNodeSelected));
+		FCanExecuteAction::CreateSP(this, &SReferenceViewer::HasExactlyOnePackageNodeSelected));
 }
 
 void SReferenceViewer::ShowSelectionInContentBrowser()
@@ -757,7 +745,10 @@ void SReferenceViewer::ShowSelectionInContentBrowser()
 	{
 		if (UEdGraphNode_Reference* ReferenceNode = Cast<UEdGraphNode_Reference>(*It))
 		{
-			AssetList.Add(ReferenceNode->GetAssetData());
+			if (ReferenceNode->GetAssetData().IsValid())
+			{
+				AssetList.Add(ReferenceNode->GetAssetData());
+			}
 		}
 	}
 
@@ -769,6 +760,23 @@ void SReferenceViewer::ShowSelectionInContentBrowser()
 
 void SReferenceViewer::OpenSelectedInAssetEditor()
 {
+	FAssetRegistryModule& AssetRegistryModule = FModuleManager::LoadModuleChecked<FAssetRegistryModule>(TEXT("AssetRegistry"));
+	TSet<UObject*> SelectedNodes = GraphEditorPtr->GetSelectedNodes();
+	for (FGraphPanelSelectionSet::TConstIterator It(SelectedNodes); It; ++It)
+	{
+		if (UEdGraphNode_Reference* ReferenceNode = Cast<UEdGraphNode_Reference>(*It))
+		{
+			if (!ReferenceNode->IsPackage() && !ReferenceNode->IsCollapsed())
+			{
+				if (AssetRegistryModule.Get().EditSearchableName(ReferenceNode->GetIdentifier()))
+				{
+					// Was handled by callback, otherwise fall back to default behavior
+					return;
+				}
+			}
+		}
+	}
+
 	UObject* SelectedObject = GetObjectFromSingleSelectedNode();
 
 	if( SelectedObject )
@@ -807,7 +815,7 @@ void SReferenceViewer::ListObjectsThatReference()
 	}
 }
 
-void SReferenceViewer::MakeCollectionWithReferenersOrDependencies(ECollectionShareType::Type ShareType, bool bReferencers)
+void SReferenceViewer::MakeCollectionWithReferencersOrDependencies(ECollectionShareType::Type ShareType, bool bReferencers)
 {
 	TSet<FName> AllSelectedPackageNames;
 	GetPackageNamesFromSelectedNodes(AllSelectedPackageNames);
@@ -914,7 +922,13 @@ void SReferenceViewer::ShowSizeMap()
 		UEdGraphNode_Reference* ReferenceNode = Cast<UEdGraphNode_Reference>(SelectedNodes.Array()[0]);
 		if ( ReferenceNode )
 		{
-			SelectedAssetPackageNames.AddUnique( ReferenceNode->GetPackageName() );
+			TArray<FName> PackageNames;
+			ReferenceNode->GetAllPackageNames(PackageNames);
+
+			for (FName Name : PackageNames)
+			{
+				SelectedAssetPackageNames.AddUnique(Name);
+			}
 		}
 	}
 
@@ -951,14 +965,14 @@ void SReferenceViewer::ShowReferenceTree()
 
 void SReferenceViewer::ReCenterGraphOnNodes(const TSet<UObject*>& Nodes)
 {
-	TArray<FName> NewGraphRootNames;
+	TArray<FAssetIdentifier> NewGraphRootNames;
 	FIntPoint TotalNodePos(EForceInit::ForceInitToZero);
 	for ( auto NodeIt = Nodes.CreateConstIterator(); NodeIt; ++NodeIt )
 	{
 		UEdGraphNode_Reference* ReferenceNode = Cast<UEdGraphNode_Reference>(*NodeIt);
 		if ( ReferenceNode )
 		{
-			NewGraphRootNames.Add(ReferenceNode->GetPackageName());
+			ReferenceNode->GetAllIdentifiers(NewGraphRootNames);
 			TotalNodePos.X += ReferenceNode->NodePosX;
 			TotalNodePos.Y += ReferenceNode->NodePosY;
 		}
@@ -1033,11 +1047,49 @@ bool SReferenceViewer::HasExactlyOneNodeSelected() const
 	return false;
 }
 
-bool SReferenceViewer::HasAtLeastOneNodeSelected() const
+bool SReferenceViewer::HasExactlyOnePackageNodeSelected() const
+{
+	if (GraphEditorPtr.IsValid())
+	{
+		if (GraphEditorPtr->GetSelectedNodes().Num() != 1)
+		{
+			return false;
+		}
+
+		TSet<UObject*> SelectedNodes = GraphEditorPtr->GetSelectedNodes();
+		for (UObject* Node : SelectedNodes)
+		{
+			UEdGraphNode_Reference* ReferenceNode = Cast<UEdGraphNode_Reference>(Node);
+			if (ReferenceNode)
+			{
+				if (ReferenceNode->IsPackage())
+				{
+					return true;
+				}
+			}
+			return false;
+		}
+	}
+
+	return false;
+}
+
+bool SReferenceViewer::HasAtLeastOnePackageNodeSelected() const
 {
 	if ( GraphEditorPtr.IsValid() )
 	{
-		return GraphEditorPtr->GetSelectedNodes().Num() > 0;
+		TSet<UObject*> SelectedNodes = GraphEditorPtr->GetSelectedNodes();
+		for (UObject* Node : SelectedNodes)
+		{
+			UEdGraphNode_Reference* ReferenceNode = Cast<UEdGraphNode_Reference>(Node);
+			if (ReferenceNode)
+			{
+				if (ReferenceNode->IsPackage())
+				{
+					return true;
+				}
+			}
+		}
 	}
 	
 	return false;
diff --git a/Engine/Source/Editor/ReferenceViewer/Private/SReferenceViewer.h b/Engine/Source/Editor/ReferenceViewer/Private/SReferenceViewer.h
--- a/Engine/Source/Editor/ReferenceViewer/Private/SReferenceViewer.h
+++ b/Engine/Source/Editor/ReferenceViewer/Private/SReferenceViewer.h
@@ -20,7 +20,7 @@ public:
 	void Construct( const FArguments& InArgs );
 
 	/** Sets a new root package name */
-	void SetGraphRootPackageNames(const TArray<FName>& NewGraphRootPackageNames);
+	void SetGraphRootPackageNames(const TArray<FAssetIdentifier>& NewGraphRootIdentifiers);
 
 private:
 
@@ -77,13 +77,13 @@ private:
 
 	void OnShowSoftReferencesChanged( ECheckBoxState NewState );
 	ECheckBoxState IsShowSoftReferencesChecked() const;
-	void OnShowSoftDependenciesChanged(ECheckBoxState NewState);
-	ECheckBoxState IsShowSoftDependenciesChecked() const;
-
 	void OnShowHardReferencesChanged(ECheckBoxState NewState);
 	ECheckBoxState IsShowHardReferencesChecked() const;
-	void OnShowHardDependenciesChanged(ECheckBoxState NewState);
-	ECheckBoxState IsShowHardDependenciesChecked() const;
+
+	void OnShowSearchableNamesChanged(ECheckBoxState NewState);
+	ECheckBoxState IsShowSearchableNamesChecked() const;
+	void OnShowNativePackagesChanged(ECheckBoxState NewState);
+	ECheckBoxState IsShowNativePackagesChecked() const;
 
 	int32 GetSearchBreadthCount() const;
 	void OnSearchBreadthCommitted(int32 NewValue);
@@ -94,7 +94,7 @@ private:
 	void ReCenterGraph();
 	void ListReferencedObjects();
 	void ListObjectsThatReference();
-	void MakeCollectionWithReferenersOrDependencies(ECollectionShareType::Type ShareType, bool bReferencers);
+	void MakeCollectionWithReferencersOrDependencies(ECollectionShareType::Type ShareType, bool bReferencers);
 	void ShowSizeMap();
 	void ShowReferenceTree();
 
@@ -103,7 +103,8 @@ private:
 	UObject* GetObjectFromSingleSelectedNode() const;
 	void GetPackageNamesFromSelectedNodes(TSet<FName>& OutNames) const;
 	bool HasExactlyOneNodeSelected() const;
-	bool HasAtLeastOneNodeSelected() const;
+	bool HasExactlyOnePackageNodeSelected() const;
+	bool HasAtLeastOnePackageNodeSelected() const;
 
 	void OnInitialAssetRegistrySearchComplete();
 private:
diff --git a/Engine/Source/Editor/ReferenceViewer/Public/ReferenceViewer.h b/Engine/Source/Editor/ReferenceViewer/Public/ReferenceViewer.h
--- a/Engine/Source/Editor/ReferenceViewer/Public/ReferenceViewer.h
+++ b/Engine/Source/Editor/ReferenceViewer/Public/ReferenceViewer.h
@@ -3,7 +3,7 @@
 #pragma once
 
 #include "ModuleManager.h"
-
+#include "AssetData.h"
 
 /**
  * The public interface to this module
@@ -35,6 +35,18 @@ public:
 	}
 
 	/** Invokes a major tab with a reference viewer within it */
-	virtual void InvokeReferenceViewerTab(const TArray<FName>& GraphRootPackageNames) = 0;
+	virtual void InvokeReferenceViewerTab(const TArray<FName>& GraphRootPackageNames)
+	{
+		TArray<FAssetIdentifier> Identifiers;
+		for (FName Name : GraphRootPackageNames)
+		{
+			Identifiers.Add(FAssetIdentifier(Name));
+		}
+
+		InvokeReferenceViewerTab(Identifiers);
+	}
+
+	/** Invokes a major tab with a reference viewer within it */
+	virtual void InvokeReferenceViewerTab(const TArray<FAssetIdentifier>& GraphRootIdentifiers) = 0;
 };
 
diff --git a/Engine/Source/Editor/SkeletonEditor/Private/EditableSkeleton.cpp b/Engine/Source/Editor/SkeletonEditor/Private/EditableSkeleton.cpp
--- a/Engine/Source/Editor/SkeletonEditor/Private/EditableSkeleton.cpp
+++ b/Engine/Source/Editor/SkeletonEditor/Private/EditableSkeleton.cpp
@@ -357,6 +357,7 @@ void FEditableSkeleton::RemoveSmartnamesAndFixupAnimations(const FName& InContai
 
 void FEditableSkeleton::SetCurveMetaDataMaterial(const FSmartName& CurveName, bool bOverrideMaterial)
 {
+	Skeleton->Modify();
 	FCurveMetaData* CurveMetaData = Skeleton->GetCurveMetaData(CurveName);
 	if (CurveMetaData)
 	{
@@ -365,6 +366,23 @@ void FEditableSkeleton::SetCurveMetaDataMaterial(const FSmartName& CurveName, bo
 	}
 }
 
+void FEditableSkeleton::SetCurveMetaBoneLinks(const FSmartName& CurveName, TArray<FBoneReference>& BoneLinks)
+{
+	Skeleton->Modify();
+	FCurveMetaData* CurveMetaData = Skeleton->GetCurveMetaData(CurveName);
+	if (CurveMetaData)
+	{
+		// override curve data
+		CurveMetaData->LinkedBones = BoneLinks;
+		
+		//  initialize to this skeleton
+		for (FBoneReference& BoneReference : CurveMetaData->LinkedBones)
+		{
+			BoneReference.Initialize(Skeleton);
+		}
+	}
+}
+
 FName FEditableSkeleton::GenerateUniqueSocketName( FName InName, USkeletalMesh* InSkeletalMesh )
 {
 	int32 NewNumber = InName.GetNumber();
diff --git a/Engine/Source/Editor/SkeletonEditor/Private/EditableSkeleton.h b/Engine/Source/Editor/SkeletonEditor/Private/EditableSkeleton.h
--- a/Engine/Source/Editor/SkeletonEditor/Private/EditableSkeleton.h
+++ b/Engine/Source/Editor/SkeletonEditor/Private/EditableSkeleton.h
@@ -35,6 +35,7 @@ public:
 	virtual void RemoveSmartname(const FName& InContainerName, SmartName::UID_Type InNameUid) override;
 	virtual void RemoveSmartnamesAndFixupAnimations(const FName& InContainerName, const TArray<SmartName::UID_Type>& InNameUids) override;
 	virtual void SetCurveMetaDataMaterial(const FSmartName& CurveName, bool bOverrideMaterial) override;
+	virtual void SetCurveMetaBoneLinks(const FSmartName& CurveName, TArray<FBoneReference>& BoneLinks) override;
 	virtual void SetPreviewMesh(class USkeletalMesh* InSkeletalMesh) override;
 	virtual void LoadAdditionalPreviewSkeletalMeshes() override;
 	virtual void SetAdditionalPreviewSkeletalMeshes(class UPreviewMeshCollection* InSkeletalMeshes) override;
diff --git a/Engine/Source/Editor/SkeletonEditor/Private/SkeletonEditorMode.cpp b/Engine/Source/Editor/SkeletonEditor/Private/SkeletonEditorMode.cpp
--- a/Engine/Source/Editor/SkeletonEditor/Private/SkeletonEditorMode.cpp
+++ b/Engine/Source/Editor/SkeletonEditor/Private/SkeletonEditorMode.cpp
@@ -16,14 +16,16 @@ FSkeletonEditorMode::FSkeletonEditorMode(TSharedRef<FWorkflowCentricApplication>
 	ISkeletonEditorModule& SkeletonEditorModule = FModuleManager::LoadModuleChecked<ISkeletonEditorModule>("SkeletonEditor");
 	TabFactories.RegisterFactory(SkeletonEditorModule.CreateSkeletonTreeTabFactory(InHostingApp, InSkeletonTree));
 
+	FOnObjectsSelected OnObjectsSelected = FOnObjectsSelected::CreateSP(&SkeletonEditor.Get(), &FSkeletonEditor::HandleObjectsSelected);
+	FOnObjectSelected OnObjectSelected = FOnObjectSelected::CreateSP(&SkeletonEditor.Get(), &FSkeletonEditor::HandleObjectSelected);
 	FPersonaModule& PersonaModule = FModuleManager::LoadModuleChecked<FPersonaModule>("Persona");
 	TabFactories.RegisterFactory(PersonaModule.CreateDetailsTabFactory(InHostingApp, FOnDetailsCreated::CreateSP(&SkeletonEditor.Get(), &FSkeletonEditor::HandleDetailsCreated)));
 	TabFactories.RegisterFactory(PersonaModule.CreatePersonaViewportTabFactory(InHostingApp, InSkeletonTree, SkeletonEditor->GetPersonaToolkit()->GetPreviewScene(), SkeletonEditor->OnPostUndo, nullptr, FOnViewportCreated(), true, true));
-	TabFactories.RegisterFactory(PersonaModule.CreateAnimNotifiesTabFactory(InHostingApp, InSkeletonTree->GetEditableSkeleton(), SkeletonEditor->OnChangeAnimNotifies, SkeletonEditor->OnPostUndo, FOnObjectsSelected::CreateSP(&SkeletonEditor.Get(), &FSkeletonEditor::HandleObjectsSelected)));
+	TabFactories.RegisterFactory(PersonaModule.CreateAnimNotifiesTabFactory(InHostingApp, InSkeletonTree->GetEditableSkeleton(), SkeletonEditor->OnChangeAnimNotifies, SkeletonEditor->OnPostUndo, OnObjectsSelected));
 	TabFactories.RegisterFactory(PersonaModule.CreateAdvancedPreviewSceneTabFactory(InHostingApp, SkeletonEditor->GetPersonaToolkit()->GetPreviewScene()));
 	TabFactories.RegisterFactory(PersonaModule.CreateRetargetManagerTabFactory(InHostingApp, InSkeletonTree->GetEditableSkeleton(), SkeletonEditor->GetPersonaToolkit()->GetPreviewScene(), SkeletonEditor->OnPostUndo));
-	TabFactories.RegisterFactory(PersonaModule.CreateCurveViewerTabFactory(InHostingApp, InSkeletonTree->GetEditableSkeleton(), SkeletonEditor->GetPersonaToolkit()->GetPreviewScene(), SkeletonEditor->OnCurvesChanged, SkeletonEditor->OnPostUndo));
-	TabFactories.RegisterFactory(PersonaModule.CreateSkeletonSlotNamesTabFactory(InHostingApp, InSkeletonTree->GetEditableSkeleton(), SkeletonEditor->OnPostUndo, FOnObjectSelected::CreateSP(&SkeletonEditor.Get(), &FSkeletonEditor::HandleObjectSelected)));
+	TabFactories.RegisterFactory(PersonaModule.CreateCurveViewerTabFactory(InHostingApp, InSkeletonTree->GetEditableSkeleton(), SkeletonEditor->GetPersonaToolkit()->GetPreviewScene(), SkeletonEditor->OnCurvesChanged, SkeletonEditor->OnPostUndo, OnObjectsSelected));
+	TabFactories.RegisterFactory(PersonaModule.CreateSkeletonSlotNamesTabFactory(InHostingApp, InSkeletonTree->GetEditableSkeleton(), SkeletonEditor->OnPostUndo, OnObjectSelected));
 
 	TabLayout = FTabManager::NewLayout("Standalone_SkeletonEditor_Layout_v1.1")
 		->AddArea
diff --git a/Engine/Source/Editor/SkeletonEditor/Public/IEditableSkeleton.h b/Engine/Source/Editor/SkeletonEditor/Public/IEditableSkeleton.h
--- a/Engine/Source/Editor/SkeletonEditor/Public/IEditableSkeleton.h
+++ b/Engine/Source/Editor/SkeletonEditor/Public/IEditableSkeleton.h
@@ -62,6 +62,10 @@ public:
 
 	/** Sets Material Meta Data for the curve */
 	virtual void SetCurveMetaDataMaterial(const FSmartName& CurveName, bool bOverrideMaterial) = 0;
+
+	/** Sets Bone Links per curve */
+	virtual void SetCurveMetaBoneLinks(const FSmartName& CurveName, TArray<FBoneReference>& BoneLinks) = 0;
+
 	/**
 	 * Makes sure all attached objects are valid and removes any that aren't.
 	 *
diff --git a/Engine/Source/Editor/UnrealEd/Classes/Animation/EditorAnimCurveBoneLinks.h b/Engine/Source/Editor/UnrealEd/Classes/Animation/EditorAnimCurveBoneLinks.h
new file mode 100644
--- /dev/null
+++ b/Engine/Source/Editor/UnrealEd/Classes/Animation/EditorAnimCurveBoneLinks.h
@@ -0,0 +1,33 @@
+// Copyright 1998-2016 Epic Games, Inc. All Rights Reserved.
+
+/**
+ * Editor only class for UI of linking animation curve to joints
+ */
+
+#pragma once
+#include "EditorAnimCurveBoneLinks.generated.h"
+
+DECLARE_DELEGATE_OneParam( FOnAnimCurveBonesChange, class UEditorAnimCurveBoneLinks*)
+
+UCLASS(MinimalAPI)
+class UEditorAnimCurveBoneLinks: public UObject
+{
+	GENERATED_UCLASS_BODY()
+public:
+	
+	virtual void Initialize(const TWeakPtr<class IEditableSkeleton> InEditableSkeleton, const FSmartName& InCurveName, FOnAnimCurveBonesChange OnChangeIn);
+
+	TWeakPtr<class IEditableSkeleton> EditableSkeleton;
+	FOnAnimCurveBonesChange OnChange;
+
+	UPROPERTY(VisibleAnywhere, Category = CurveName)
+	FSmartName CurveName;
+
+	UPROPERTY(EditAnywhere, Category = Bones)
+	TArray<FBoneReference> ConnectedBones;
+
+	virtual void PostEditChangeProperty(FPropertyChangedEvent& PropertyChangedEvent) override;
+
+	// refresh current Connected Bones data 
+	UNREALED_API void Refresh(const FSmartName& InCurveName, const TArray<FBoneReference>& CurrentLinks);
+};
diff --git a/Engine/Source/Editor/UnrealEd/Classes/Settings/EditorExperimentalSettings.h b/Engine/Source/Editor/UnrealEd/Classes/Settings/EditorExperimentalSettings.h
--- a/Engine/Source/Editor/UnrealEd/Classes/Settings/EditorExperimentalSettings.h
+++ b/Engine/Source/Editor/UnrealEd/Classes/Settings/EditorExperimentalSettings.h
@@ -75,6 +75,10 @@ public:
 	UPROPERTY(EditAnywhere, config, Category=Audio)
 	bool bShowAudioStreamingOptions;
 
+	/** Whether to show AudioMixer-dependent editor data. Only enable if also running with the -audiomixer. */
+	UPROPERTY(EditAnywhere, config, Category = Audio)
+	bool bShowAudioMixerData;
+
 	/** Allows ChunkIDs to be assigned to assets to via the content browser context menu. */
 	UPROPERTY(EditAnywhere,config,Category=UserInterface,meta=(DisplayName="Allow ChunkID Assignments"))
 	bool bContextMenuChunkAssignments;
diff --git a/Engine/Source/Editor/UnrealEd/Private/Commandlets/ChunkManifestGenerator.cpp b/Engine/Source/Editor/UnrealEd/Private/Commandlets/ChunkManifestGenerator.cpp
--- a/Engine/Source/Editor/UnrealEd/Private/Commandlets/ChunkManifestGenerator.cpp
+++ b/Engine/Source/Editor/UnrealEd/Private/Commandlets/ChunkManifestGenerator.cpp
@@ -79,7 +79,7 @@ FChunkManifestGenerator::FChunkManifestGenerator(const TArray<ITargetPlatform*>&
 	}	
 
 	UE_LOG(LogChunkManifestGenerator, Log, TEXT("bChunkHardReferencesOnly: %i"), (int32)bOnlyHardReferences);
-	DependencyType = bOnlyHardReferences ? EAssetRegistryDependencyType::Hard : EAssetRegistryDependencyType::All;
+	DependencyType = bOnlyHardReferences ? EAssetRegistryDependencyType::Hard : EAssetRegistryDependencyType::Packages;
 }
 
 FChunkManifestGenerator::~FChunkManifestGenerator()
diff --git a/Engine/Source/Editor/UnrealEd/Private/CookOnTheFlyServer.cpp b/Engine/Source/Editor/UnrealEd/Private/CookOnTheFlyServer.cpp
--- a/Engine/Source/Editor/UnrealEd/Private/CookOnTheFlyServer.cpp
+++ b/Engine/Source/Editor/UnrealEd/Private/CookOnTheFlyServer.cpp
@@ -2620,7 +2620,7 @@ const TArray<FName>& UCookOnTheFlyServer::GetFullPackageDependencies(const FName
 		CachedFullPackageDependencies.Add(PackageName);
 
 		TArray<FName> ChildDependencies;
-		check( AssetRegistry.GetDependencies(PackageName, ChildDependencies, EAssetRegistryDependencyType::All) );
+		check( AssetRegistry.GetDependencies(PackageName, ChildDependencies, EAssetRegistryDependencyType::Packages) );
 
 		TArray<FName> Dependencies = ChildDependencies;
 		Dependencies.AddUnique(PackageName);
diff --git a/Engine/Source/Editor/UnrealEd/Private/EditorEngine.cpp b/Engine/Source/Editor/UnrealEd/Private/EditorEngine.cpp
--- a/Engine/Source/Editor/UnrealEd/Private/EditorEngine.cpp
+++ b/Engine/Source/Editor/UnrealEd/Private/EditorEngine.cpp
@@ -705,7 +705,6 @@ void UEditorEngine::Init(IEngineLoop* InEngineLoop)
 			TEXT("Blutility"),
 			TEXT("XmlParser"),
 			TEXT("UserFeedback"),
-			TEXT("GameplayTagsEditor"),
 			TEXT("UndoHistory"),
 			TEXT("DeviceProfileEditor"),
 			TEXT("SourceCodeAccess"),
diff --git a/Engine/Source/Editor/UnrealEd/Private/EditorModeManager.cpp b/Engine/Source/Editor/UnrealEd/Private/EditorModeManager.cpp
--- a/Engine/Source/Editor/UnrealEd/Private/EditorModeManager.cpp
+++ b/Engine/Source/Editor/UnrealEd/Private/EditorModeManager.cpp
@@ -110,6 +110,11 @@ TSharedPtr<class IToolkitHost> FEditorModeTools::GetToolkitHost() const
 	return Result;
 }
 
+bool FEditorModeTools::HasToolkitHost() const
+{
+	return ToolkitHost.Pin().IsValid();
+}
+
 void FEditorModeTools::SetToolkitHost(TSharedRef<class IToolkitHost> InHost)
 {
 	checkf(!ToolkitHost.IsValid(), TEXT("SetToolkitHost can only be called once"));
diff --git a/Engine/Source/Editor/UnrealEd/Private/EditorViewportClient.cpp b/Engine/Source/Editor/UnrealEd/Private/EditorViewportClient.cpp
--- a/Engine/Source/Editor/UnrealEd/Private/EditorViewportClient.cpp
+++ b/Engine/Source/Editor/UnrealEd/Private/EditorViewportClient.cpp
@@ -278,6 +278,7 @@ FEditorViewportClient::FEditorViewportClient(FEditorModeTools* InModeTools, FPre
 	, bNeedsLinkedRedraw(false)
 	, bNeedsInvalidateHitProxy(false)
 	, bUsingOrbitCamera(false)
+	, bUseNumpadCameraControl(true)
 	, bDisableInput(false)
 	, bDrawAxes(true)
 	, bSetListenerPosition(false)
@@ -1526,13 +1527,13 @@ void FEditorViewportClient::UpdateCameraMovement( float DeltaTime )
 		// Forward/back
 		if( ( bRemapWASDKeys && Viewport->KeyState( FViewportNavigationCommands::Get().Forward->GetActiveChord()->Key ) ) ||
 			( bRemapArrowKeys && Viewport->KeyState( EKeys::Up ) ) ||
-			( bUnmodifiedPress && Viewport->KeyState(EKeys::NumPadEight) ) )
+			( bUnmodifiedPress && bUseNumpadCameraControl && Viewport->KeyState(EKeys::NumPadEight) ) )
 		{
 			CameraUserImpulseData->MoveForwardBackwardImpulse += 1.0f;
 		}
 		if( ( bRemapWASDKeys && Viewport->KeyState( FViewportNavigationCommands::Get().Backward->GetActiveChord()->Key ) ) ||
 			( bRemapArrowKeys && Viewport->KeyState( EKeys::Down ) ) ||
-			( bUnmodifiedPress && Viewport->KeyState( EKeys::NumPadTwo ) ) )
+			( bUnmodifiedPress && bUseNumpadCameraControl && Viewport->KeyState( EKeys::NumPadTwo ) ) )
 		{
 			CameraUserImpulseData->MoveForwardBackwardImpulse -= 1.0f;
 		}
@@ -1540,37 +1541,39 @@ void FEditorViewportClient::UpdateCameraMovement( float DeltaTime )
 		// Right/left
 		if ( ( bRemapWASDKeys && Viewport->KeyState( FViewportNavigationCommands::Get().Right->GetActiveChord()->Key) ) ||
 			( bRemapArrowKeys && Viewport->KeyState( EKeys::Right ) ) ||
-			( bUnmodifiedPress && Viewport->KeyState( EKeys::NumPadSix ) ) )
+			( bUnmodifiedPress && bUseNumpadCameraControl && Viewport->KeyState( EKeys::NumPadSix ) ) )
 		{
 			CameraUserImpulseData->MoveRightLeftImpulse += 1.0f;
 		}
 		if( ( bRemapWASDKeys && Viewport->KeyState( FViewportNavigationCommands::Get().Left->GetActiveChord()->Key ) ) ||
 			( bRemapArrowKeys && Viewport->KeyState( EKeys::Left ) ) ||
-			( bUnmodifiedPress && Viewport->KeyState( EKeys::NumPadFour ) ) )
+			( bUnmodifiedPress && bUseNumpadCameraControl && Viewport->KeyState( EKeys::NumPadFour ) ) )
 		{
 			CameraUserImpulseData->MoveRightLeftImpulse -= 1.0f;
 		}
 
 		// Up/down
 		if( ( bRemapWASDKeys && Viewport->KeyState( FViewportNavigationCommands::Get().Up->GetActiveChord()->Key ) ) ||
-			( bUnmodifiedPress && ( Viewport->KeyState( EKeys::PageUp ) || Viewport->KeyState( EKeys::NumPadNine ) || Viewport->KeyState( EKeys::Add ) ) ) )
+			( bUnmodifiedPress && Viewport->KeyState( EKeys::PageUp ) ) ||
+			( bUnmodifiedPress && bUseNumpadCameraControl && ( Viewport->KeyState( EKeys::NumPadNine ) || Viewport->KeyState( EKeys::Add ) ) ) )
 		{
 			CameraUserImpulseData->MoveUpDownImpulse += 1.0f;
 		}
 		if( ( bRemapWASDKeys && Viewport->KeyState( FViewportNavigationCommands::Get().Down->GetActiveChord()->Key) ) ||
-			( bUnmodifiedPress && ( Viewport->KeyState( EKeys::PageDown ) || Viewport->KeyState( EKeys::NumPadSeven ) || Viewport->KeyState( EKeys::Subtract ) ) ) )
+			( bUnmodifiedPress && Viewport->KeyState( EKeys::PageDown ) ) ||
+			( bUnmodifiedPress && bUseNumpadCameraControl && ( Viewport->KeyState( EKeys::NumPadSeven ) || Viewport->KeyState( EKeys::Subtract ) ) ) )
 		{
 			CameraUserImpulseData->MoveUpDownImpulse -= 1.0f;
 		}
 
 		// Zoom FOV out/in
 		if( ( bRemapWASDKeys && Viewport->KeyState( FViewportNavigationCommands::Get().FovZoomOut->GetActiveChord()->Key ) ) ||
-			( bUnmodifiedPress && Viewport->KeyState( EKeys::NumPadOne ) ) )
+			( bUnmodifiedPress && bUseNumpadCameraControl && Viewport->KeyState( EKeys::NumPadOne ) ) )
 		{
 			CameraUserImpulseData->ZoomOutInImpulse += 1.0f;
 		}
 		if( ( bRemapWASDKeys && Viewport->KeyState( FViewportNavigationCommands::Get().FovZoomIn->GetActiveChord()->Key ) ) ||
-			( bUnmodifiedPress && Viewport->KeyState( EKeys::NumPadThree ) ) )
+			( bUnmodifiedPress && bUseNumpadCameraControl && Viewport->KeyState( EKeys::NumPadThree ) ) )
 		{
 			CameraUserImpulseData->ZoomOutInImpulse -= 1.0f;
 		}
diff --git a/Engine/Source/Editor/UnrealEd/Private/FbxMeshUtils.cpp b/Engine/Source/Editor/UnrealEd/Private/FbxMeshUtils.cpp
--- a/Engine/Source/Editor/UnrealEd/Private/FbxMeshUtils.cpp
+++ b/Engine/Source/Editor/UnrealEd/Private/FbxMeshUtils.cpp
@@ -69,8 +69,10 @@ namespace FbxMeshUtils
 		}
 	}
 
-	void ImportStaticMeshLOD( UStaticMesh* BaseStaticMesh, const FString& Filename, int32 LODLevel)
+	bool ImportStaticMeshLOD( UStaticMesh* BaseStaticMesh, const FString& Filename, int32 LODLevel)
 	{
+		bool bSuccess = false;
+
 		UE_LOG(LogExportMeshUtils, Log, TEXT("Fbx LOD loading"));
 		// logger for all error/warnings
 		// this one prints all messages that are stored in FFbxImporter
@@ -132,7 +134,7 @@ namespace FbxMeshUtils
 					FFbxImporter->AddTokenizedErrorMessage(FTokenizedMessage::Create(EMessageSeverity::Error, FText(LOCTEXT("Prompt_NoMeshFound", "No meshes were found in file."))), FFbxErrors::Generic_Mesh_MeshNotFound);
 
 					FFbxImporter->ReleaseScene();
-					return;
+					return bSuccess;
 				}
 			}
 
@@ -158,6 +160,8 @@ namespace FbxMeshUtils
 					NotificationInfo.Text = FText::Format(LOCTEXT("LODImportSuccessful", "Mesh for LOD {0} imported successfully!"), FText::AsNumber(LODLevel));
 					NotificationInfo.ExpireDuration = 5.0f;
 					FSlateNotificationManager::Get().AddNotification(NotificationInfo);
+
+					bSuccess = true;
 				}
 				else
 				{
@@ -166,6 +170,8 @@ namespace FbxMeshUtils
 					NotificationInfo.Text = FText::Format(LOCTEXT("LODImportFail", "Failed to import mesh for LOD {0}!"), FText::AsNumber( LODLevel ));
 					NotificationInfo.ExpireDuration = 5.0f;
 					FSlateNotificationManager::Get().AddNotification(NotificationInfo);
+
+					bSuccess = false;
 				}
 			}
 
@@ -176,10 +182,14 @@ namespace FbxMeshUtils
 			}
 		}
 		FFbxImporter->ReleaseScene();
+
+		return bSuccess;
 	}
 
-	void ImportSkeletalMeshLOD( class USkeletalMesh* SelectedSkelMesh, const FString& Filename, int32 LODLevel)
+	bool ImportSkeletalMeshLOD( class USkeletalMesh* SelectedSkelMesh, const FString& Filename, int32 LODLevel)
 	{
+		bool bSuccess = false;
+
 		// Check the file extension for FBX. Anything that isn't .FBX is rejected
 		const FString FileExtension = FPaths::GetExtension(Filename);
 		const bool bIsFBX = FCString::Stricmp(*FileExtension, TEXT("FBX")) == 0;
@@ -248,7 +258,7 @@ namespace FbxMeshUtils
 				{
 					FFbxImporter->AddTokenizedErrorMessage(FTokenizedMessage::Create(EMessageSeverity::Error, LOCTEXT("FBXImport_NoMesh", "No meshes were found in file.")), FFbxErrors::Generic_MeshNotFound);
 					FFbxImporter->ReleaseScene();
-					return;
+					return false;
 				}
 
 				MeshObject = MeshArray[0];
@@ -343,10 +353,10 @@ namespace FbxMeshUtils
 					TempSkelMesh = (USkeletalMesh*)FFbxImporter->ImportSkeletalMesh(SelectedSkelMesh->GetOutermost(), bUseLODs? SkelMeshNodeArray: *MeshObject, NAME_None, RF_Transient, TempAssetImportData, LODLevel, nullptr, nullptr, nullptr, true, OrderedMaterialNames.Num() > 0 ? &OrderedMaterialNames : nullptr);
 
 					// Add imported mesh to existing model
-					bool bImportSucceeded = false;
+					bool bMeshImportSuccess = false;
 					if( TempSkelMesh )
 					{
-						bImportSucceeded = FFbxImporter->ImportSkeletalMeshLOD(TempSkelMesh, SelectedSkelMesh, SelectedLOD);
+						bMeshImportSuccess = FFbxImporter->ImportSkeletalMeshLOD(TempSkelMesh, SelectedSkelMesh, SelectedLOD);
 
 						// Mark package containing skeletal mesh as dirty.
 						SelectedSkelMesh->MarkPackageDirty();
@@ -357,8 +367,12 @@ namespace FbxMeshUtils
 						FFbxImporter->ImportFbxMorphTarget(SkelMeshNodeArray, SelectedSkelMesh, SelectedSkelMesh->GetOutermost(), SelectedLOD);
 					}
 
-					if (bImportSucceeded)
+					if (bMeshImportSuccess)
 					{
+						bSuccess = true;
+						// Set LOD source filename
+						SelectedSkelMesh->LODInfo[SelectedLOD].SourceImportFilename = Filename;
+
 						// Notification of success
 						FNotificationInfo NotificationInfo(FText::GetEmpty());
 						NotificationInfo.Text = FText::Format(NSLOCTEXT("UnrealEd", "LODImportSuccessful", "Mesh for LOD {0} imported successfully!"), FText::AsNumber(SelectedLOD));
@@ -391,42 +405,29 @@ namespace FbxMeshUtils
 			ApexClothingUtils::ReImportClothingSectionsFromClothingAsset(SelectedSkelMesh);
 #endif// #if WITH_APEX_CLOTHING
 		}
+
+		return bSuccess;
 	}
 
-	void ImportMeshLODDialog( class UObject* SelectedMesh, int32 LODLevel)
+	FString PromptForLODImportFile(const FText& PromptTitle)
 	{
-		if(!SelectedMesh)
-		{
-			return;
-		}
-
-		USkeletalMesh* SkeletonMesh = Cast<USkeletalMesh>(SelectedMesh);
-		UStaticMesh* StaticMesh = Cast<UStaticMesh>(SelectedMesh);
-
-		if( !SkeletonMesh && !StaticMesh )
-		{
-			return;
-		}
+		FString ChosenFilname("");
 
 		FString ExtensionStr;
-
 		ExtensionStr += TEXT("All model files|*.fbx;*.obj|");
-
 		ExtensionStr += TEXT("FBX files|*.fbx|");
-
 		ExtensionStr += TEXT("Object files|*.obj|");
-
 		ExtensionStr += TEXT("All files|*.*");
 
 		// First, display the file open dialog for selecting the file.
 		TArray<FString> OpenFilenames;
 		IDesktopPlatform* DesktopPlatform = FDesktopPlatformModule::Get();
 		bool bOpen = false;
-		if ( DesktopPlatform )
+		if(DesktopPlatform)
 		{
 			bOpen = DesktopPlatform->OpenFileDialog(
 				FSlateApplication::Get().FindBestParentWindowHandleForDialogs(nullptr),
-				FText::Format( NSLOCTEXT("UnrealEd", "ImportMeshLOD", "Failed to import mesh for LOD {0}!"), FText::AsNumber( LODLevel ) ).ToString(),
+				PromptTitle.ToString(),
 				*FEditorDirectories::Get().GetLastDirectory(ELastDirectory::FBX),
 				TEXT(""),
 				*ExtensionStr,
@@ -434,11 +435,11 @@ namespace FbxMeshUtils
 				OpenFilenames
 				);
 		}
-			
+
 		// Only continue if we pressed OK and have only one file selected.
-		if( bOpen )
+		if(bOpen)
 		{
-			if( OpenFilenames.Num() == 0)
+			if(OpenFilenames.Num() == 0)
 			{
 				UnFbx::FFbxImporter* FFbxImporter = UnFbx::FFbxImporter::GetInstance();
 				FFbxImporter->AddTokenizedErrorMessage(FTokenizedMessage::Create(EMessageSeverity::Error, LOCTEXT("NoFileSelectedForLOD", "No file was selected for the LOD.")), FFbxErrors::Generic_Mesh_LOD_NoFileSelected);
@@ -450,19 +451,102 @@ namespace FbxMeshUtils
 			}
 			else
 			{
-				FString Filename = OpenFilenames[0];
-				FEditorDirectories::Get().SetLastDirectory(ELastDirectory::FBX, FPaths::GetPath(Filename)); // Save path as default for next time.
+				ChosenFilname = OpenFilenames[0];
+				FEditorDirectories::Get().SetLastDirectory(ELastDirectory::FBX, FPaths::GetPath(ChosenFilname)); // Save path as default for next time.
+			}
+		}
+		
+		return ChosenFilname;
+	}
 
-				if( SkeletonMesh )
+	bool ImportMeshLODDialog(class UObject* SelectedMesh, int32 LODLevel)
+	{
+		if(!SelectedMesh)
+		{
+			return false;
+		}
+
+		USkeletalMesh* SkeletalMesh = Cast<USkeletalMesh>(SelectedMesh);
+		UStaticMesh* StaticMesh = Cast<UStaticMesh>(SelectedMesh);
+
+		if( !SkeletalMesh && !StaticMesh )
+		{
+			return false;
+		}
+
+		FString FilenameToImport("");
+
+		if(SkeletalMesh)
+		{
+			if(SkeletalMesh->LODInfo.IsValidIndex(LODLevel))
+			{
+				FSkeletalMeshLODInfo& SkelLodInfo = SkeletalMesh->LODInfo[LODLevel];
+				FilenameToImport = SkelLodInfo.SourceImportFilename;
+			}
+		}
+
+		// Check the file exists first
+		const bool bSourceFileExists = FPaths::FileExists(FilenameToImport);
+		// We'll give the user a chance to choose a new file if a previously set file fails to import
+		const bool bPromptOnFail = bSourceFileExists;
+		
+		if(!bSourceFileExists || FilenameToImport.IsEmpty())
+		{
+			FText PromptTitle;
+
+			if(FilenameToImport.IsEmpty())
+			{
+				PromptTitle = FText::Format(LOCTEXT("LODImportPrompt_NoSource", "Choose a file to import for LOD {0}"), FText::AsNumber(LODLevel));
+			}
+			else if(!bSourceFileExists)
+			{
+				PromptTitle = FText::Format(LOCTEXT("LODImportPrompt_SourceNotFound", "LOD {0} Source file not found. Choose new file."), FText::AsNumber(LODLevel));
+			}
+
+			FilenameToImport = PromptForLODImportFile(PromptTitle);
+		}
+		
+		bool bImportSuccess = false;
+
+		if(!FilenameToImport.IsEmpty())
+		{
+			if(SkeletalMesh)
+			{
+				bImportSuccess = ImportSkeletalMeshLOD(SkeletalMesh, FilenameToImport, LODLevel);
+			}
+			else if(StaticMesh)
+			{
+				bImportSuccess = ImportStaticMeshLOD(StaticMesh, FilenameToImport, LODLevel);
+			}
+		}
+
+		if(!bImportSuccess && bPromptOnFail)
+		{
+			FMessageDialog::Open(EAppMsgType::Ok, FText::Format(LOCTEXT("LODImport_SourceMissingDialog", "Failed to import LOD{0} as the source file failed to import, please select a new source file."), FText::AsNumber(LODLevel)));
+
+			FText PromptTitle = FText::Format(LOCTEXT("LODImportPrompt_SourceFailed", "Failed to import source file for LOD {0}, choose a new file"), FText::AsNumber(LODLevel));
+			FilenameToImport = PromptForLODImportFile(PromptTitle);
+
+			if(FilenameToImport.Len() > 0 && FPaths::FileExists(FilenameToImport))
+			{
+				if(SkeletalMesh)
 				{
-					ImportSkeletalMeshLOD(SkeletonMesh, Filename, LODLevel);
+					bImportSuccess = ImportSkeletalMeshLOD(SkeletalMesh, FilenameToImport, LODLevel);
 				}
-				else if( StaticMesh )
+				else if(StaticMesh)
 				{
-					ImportStaticMeshLOD(StaticMesh, Filename, LODLevel);
+					bImportSuccess = ImportStaticMeshLOD(StaticMesh, FilenameToImport, LODLevel);
 				}
 			}
 		}
+
+		if(!bImportSuccess)
+		{
+			// Failed to import a LOD, even after retries (if applicable)
+			FMessageDialog::Open(EAppMsgType::Ok, FText::Format(LOCTEXT("LODImport_Failure", "Failed to import LOD{0}"), FText::AsNumber(LODLevel)));
+		}
+
+		return bImportSuccess;
 	}
 
 	void SetImportOption(UFbxImportUI* ImportUI)
diff --git a/Engine/Source/Editor/UnrealEd/Private/Kismet2/BlueprintEditorUtils.cpp b/Engine/Source/Editor/UnrealEd/Private/Kismet2/BlueprintEditorUtils.cpp
--- a/Engine/Source/Editor/UnrealEd/Private/Kismet2/BlueprintEditorUtils.cpp
+++ b/Engine/Source/Editor/UnrealEd/Private/Kismet2/BlueprintEditorUtils.cpp
@@ -528,6 +528,18 @@ void FBlueprintEditorUtils::ReconstructAllNodes(UBlueprint* Blueprint)
 	}
 }
 
+void FBlueprintEditorUtils::ReplaceDeprecatedNodes(UBlueprint* Blueprint)
+{
+	TArray<UEdGraph*> Graphs;
+	Blueprint->GetAllGraphs(Graphs);
+	for (auto It = Graphs.CreateIterator(); It; ++It)
+	{
+		UEdGraph* const Graph = *It;
+		const UEdGraphSchema* Schema = Graph->GetSchema();
+		Schema->BackwardCompatibilityNodeConversion(Graph, true);
+	}
+}
+
 void FBlueprintEditorUtils::RefreshExternalBlueprintDependencyNodes(UBlueprint* Blueprint, UStruct* RefreshOnlyChild)
 {
 	BP_SCOPED_COMPILER_EVENT_STAT(EKismetCompilerStats_RefreshExternalDependencyNodes);
@@ -1577,6 +1589,8 @@ UClass* FBlueprintEditorUtils::RegenerateBlueprintClass(UBlueprint* Blueprint, U
 			// and it's companion function UK2Node_BaseAsyncTask::ExpandNode.
 			FBlueprintEditorUtils::ReconstructAllNodes(Blueprint);
 
+			FBlueprintEditorUtils::ReplaceDeprecatedNodes(Blueprint);
+
 			// Compile the actual blueprint
 			FKismetEditorUtilities::CompileBlueprint(Blueprint, true, false, false, nullptr, bSkeletonUpToDate);
 		}
@@ -1585,6 +1599,8 @@ UClass* FBlueprintEditorUtils::RegenerateBlueprintClass(UBlueprint* Blueprint, U
 			// Just refresh all nodes in macro blueprints, but don't recompile
 			FBlueprintEditorUtils::RefreshAllNodes(Blueprint);
 
+			FBlueprintEditorUtils::ReplaceDeprecatedNodes(Blueprint);
+
 			if (ClassToRegenerate != nullptr)
 			{
 				UClass* OldSuperClass = ClassToRegenerate->GetSuperClass();
@@ -1893,17 +1909,14 @@ void FBlueprintEditorUtils::PatchCDOSubobjectsIntoExport(UObject* PreviousCDO, U
 		{
 			static void PatchSubObjects(UObject* OldObj, UObject* NewObj)
 			{
-				TArray<UObject*> NewSubObjects;
-				GetObjectsWithOuter(NewObj, NewSubObjects, /*bIncludeNestedSubObjects =*/false);
-
 				TMap<FName, UObject*> SubObjLookupTable;
-				for (UObject* NewSubObj : NewSubObjects)
+				ForEachObjectWithOuter(NewObj, [&SubObjLookupTable](UObject* NewSubObj)
 				{
 					if (NewSubObj != nullptr)
 					{
 						SubObjLookupTable.Add(NewSubObj->GetFName(), NewSubObj);
 					}
-				}
+				}, /*bIncludeNestedSubObjects =*/false);
 
 				TArray<UObject*> OldSubObjects;
 				GetObjectsWithOuter(OldObj, OldSubObjects, /*bIncludeNestedSubObjects =*/false);
diff --git a/Engine/Source/Editor/UnrealEd/Public/EditorModeManager.h b/Engine/Source/Editor/UnrealEd/Public/EditorModeManager.h
--- a/Engine/Source/Editor/UnrealEd/Public/EditorModeManager.h
+++ b/Engine/Source/Editor/UnrealEd/Public/EditorModeManager.h
@@ -384,6 +384,9 @@ public:
 	/** Returns the host for toolkits created via modes from this mode manager */
 	TSharedPtr<IToolkitHost> GetToolkitHost() const;
 
+	/** Check if toolkit host exists */
+	bool HasToolkitHost() const;
+
 	/**
 	 * Returns the set of selected actors.
 	 */
diff --git a/Engine/Source/Editor/UnrealEd/Public/EditorViewportClient.h b/Engine/Source/Editor/UnrealEd/Public/EditorViewportClient.h
--- a/Engine/Source/Editor/UnrealEd/Public/EditorViewportClient.h
+++ b/Engine/Source/Editor/UnrealEd/Public/EditorViewportClient.h
@@ -1317,6 +1317,8 @@ public:
 	/** True if the orbit camera is currently being used */
 	bool bUsingOrbitCamera;
 	
+	/** If true, numpad keys will be used to move camera in perspective viewport */
+	bool bUseNumpadCameraControl;
 
 	/**
 	 * true if all input is rejected from this viewport
diff --git a/Engine/Source/Editor/UnrealEd/Public/FbxMeshUtils.h b/Engine/Source/Editor/UnrealEd/Public/FbxMeshUtils.h
--- a/Engine/Source/Editor/UnrealEd/Public/FbxMeshUtils.h
+++ b/Engine/Source/Editor/UnrealEd/Public/FbxMeshUtils.h
@@ -10,8 +10,9 @@ namespace FbxMeshUtils
 	 * @param BaseStaticMesh	The static mesh to import the LOD to
 	 * @param Filename			The filename of the FBX file containing the LOD
 	 * @param LODLevel			The level of the lod to import
+	 * @return					Whether or not the import succeeded
 	 */
-	UNREALED_API void ImportStaticMeshLOD( UStaticMesh* BaseStaticMesh, const FString& Filename, int32 LODLevel );
+	UNREALED_API bool ImportStaticMeshLOD( UStaticMesh* BaseStaticMesh, const FString& Filename, int32 LODLevel );
 
 	/**
 	 * Imports a skeletal mesh LOD to the given skeletal mesh
@@ -19,8 +20,9 @@ namespace FbxMeshUtils
 	 * @param Mesh				The skeletal mesh to import the LOD to
 	 * @param Filename			The filename of the FBX file containing the LOD
 	 * @param LODLevel			The level of the lod to import
+	 * @return					Whether or not the import succeeded
 	 */
-	UNREALED_API void ImportSkeletalMeshLOD( class USkeletalMesh* Mesh, const FString& Filename, int32 LODLevel );
+	UNREALED_API bool ImportSkeletalMeshLOD( class USkeletalMesh* Mesh, const FString& Filename, int32 LODLevel );
 
 	/**
 	 * Imports a skeletal mesh LOD to the given skeletal mesh
@@ -28,8 +30,9 @@ namespace FbxMeshUtils
 	 *
 	 * @param Mesh				The skeletal mesh to import the LOD to
 	 * @param LODLevel			The level of the lod to import
+	 * @return					Whether or not the import succeeded
 	 */
-	UNREALED_API void ImportMeshLODDialog( class UObject* Mesh, int32 LODLevel );
+	UNREALED_API bool ImportMeshLODDialog( class UObject* Mesh, int32 LODLevel);
 
 	/**
 	 * Sets import option before importing
diff --git a/Engine/Source/Editor/UnrealEd/Public/Kismet2/BlueprintEditorUtils.h b/Engine/Source/Editor/UnrealEd/Public/Kismet2/BlueprintEditorUtils.h
--- a/Engine/Source/Editor/UnrealEd/Public/Kismet2/BlueprintEditorUtils.h
+++ b/Engine/Source/Editor/UnrealEd/Public/Kismet2/BlueprintEditorUtils.h
@@ -137,6 +137,11 @@ public:
 	 */
 	static void RefreshGraphNodes(const UEdGraph* Graph);
 
+	/**
+	 * Replaces any deprecated nodes with new ones
+	 */
+	static void ReplaceDeprecatedNodes(UBlueprint* Blueprint);
+
 	/**
 	 * Preloads the object and all the members it owns (nodes, pins, etc)
 	 */
diff --git a/Engine/Source/Editor/UnrealEd/UnrealEd.Build.cs b/Engine/Source/Editor/UnrealEd/UnrealEd.Build.cs
--- a/Engine/Source/Editor/UnrealEd/UnrealEd.Build.cs
+++ b/Engine/Source/Editor/UnrealEd/UnrealEd.Build.cs
@@ -207,9 +207,7 @@ public class UnrealEd : ModuleRules
                 "EnvironmentQueryEditor",
 				"ViewportSnapping",
 				"UserFeedback",
-				"GameplayTagsEditor",
                 "GameplayTasksEditor",
-                "GameplayAbilitiesEditor",
 				"UndoHistory",
 				"SourceCodeAccess",
 				"ReferenceViewer",
diff --git a/Engine/Source/Runtime/AIModule/Classes/BehaviorTree/BlackboardComponent.h b/Engine/Source/Runtime/AIModule/Classes/BehaviorTree/BlackboardComponent.h
--- a/Engine/Source/Runtime/AIModule/Classes/BehaviorTree/BlackboardComponent.h
+++ b/Engine/Source/Runtime/AIModule/Classes/BehaviorTree/BlackboardComponent.h
@@ -75,10 +75,18 @@ public:
 	/** unregister all observers associated with given owner */
 	void UnregisterObserversFrom(UObject* NotifyOwner);
 
+	/** pause observer change notifications, any new ones will be added to a queue */
+	void PauseObserverNotifications();
+
+	/** resume observer change notifications and, optionally, process the queued observation list */
+	void ResumeObserverNotifications(bool bSendQueuedObserverNotifications);
+
 	/** pause change notifies and add them to queue */
+	DEPRECATED(4.15, "Please call PauseObserverUpdates.")
 	void PauseUpdates();
 
 	/** resume change notifies and process queued list */
+	DEPRECATED(4.15, "Please call ResumeObserverNotifications.")
 	void ResumeUpdates();
 
 	/** @return associated behavior tree component */
@@ -243,7 +251,7 @@ protected:
 	/** queued key change notification, will be processed on ResumeUpdates call */
 	mutable TArray<uint8> QueuedUpdates;
 
-	/** set when notifies are paused and shouldn't be passed to observers */
+	/** set when observation notifies are paused and shouldn't be passed to observers */
 	uint32 bPausedNotifies : 1;
 
 	/** reset to false every time a new BB asset is assigned to this component */
diff --git a/Engine/Source/Runtime/AIModule/Private/BehaviorTree/BehaviorTreeComponent.cpp b/Engine/Source/Runtime/AIModule/Private/BehaviorTree/BehaviorTreeComponent.cpp
--- a/Engine/Source/Runtime/AIModule/Private/BehaviorTree/BehaviorTreeComponent.cpp
+++ b/Engine/Source/Runtime/AIModule/Private/BehaviorTree/BehaviorTreeComponent.cpp
@@ -72,7 +72,7 @@ void UBehaviorTreeComponent::PauseLogic(const FString& Reason)
 
 	if (BlackboardComp)
 	{
-		BlackboardComp->PauseUpdates();
+		BlackboardComp->PauseObserverNotifications();
 	}
 }
 
@@ -87,7 +87,8 @@ EAILogicResuming::Type UBehaviorTreeComponent::ResumeLogic(const FString& Reason
 		{
 			if (BlackboardComp)
 			{
-				BlackboardComp->ResumeUpdates();
+				// Resume the blackboard's observer notifications and send any queued notifications
+				BlackboardComp->ResumeObserverNotifications(true);
 			}
 
 			const bool bOutOfNodesPending = PendingExecution.IsSet() && PendingExecution.bOutOfNodes;
@@ -98,6 +99,14 @@ EAILogicResuming::Type UBehaviorTreeComponent::ResumeLogic(const FString& Reason
 
 			return EAILogicResuming::Continue;
 		}
+		else if (SuperResumeResult == EAILogicResuming::RestartedInstead)
+		{
+			if (BlackboardComp)
+			{
+				// Resume the blackboard's observer notifications but do not send any queued notifications
+				BlackboardComp->ResumeObserverNotifications(false);
+			}
+		}
 	}
 
 	return SuperResumeResult;
@@ -860,15 +869,13 @@ void UBehaviorTreeComponent::RequestExecution(UBTCompositeNode* RequestedOn, int
 	// store it
 	StoreDebuggerRestart(DebuggerNode, InstanceIdx, true);
 
-	if (bSwitchToHigherPriority)
+	// search end can be set only when switching to high priority
+	// or previous request was limited and current limit is wider
+	if ((!bAlreadyHasRequest && bSwitchToHigherPriority) ||
+		(ExecutionRequest.SearchEnd.IsSet() && ExecutionRequest.SearchEnd.TakesPriorityOver(SearchEnd)))
 	{
-		// search end can be set only when there wasn't any previous request or
-		// previous request was limited and current limit is wider
-		if (!bAlreadyHasRequest ||
-			(ExecutionRequest.SearchEnd.IsSet() && ExecutionRequest.SearchEnd.TakesPriorityOver(SearchEnd)))
-		{
-			ExecutionRequest.SearchEnd = SearchEnd;
-		}
+		UE_CVLOG(bAlreadyHasRequest, GetOwner(), LogBehaviorTree, Log, (SearchEnd.ExecutionIndex < MAX_uint16) ? TEXT("> expanding end of search range!") : TEXT("> removing limit from end of search range!"));
+		ExecutionRequest.SearchEnd = SearchEnd;
 	}
 
 	ExecutionRequest.SearchStart = ExecutionIdx;
@@ -1247,15 +1254,11 @@ void UBehaviorTreeComponent::ProcessExecutionRequest()
 			// check it's the same as active node (or any of active parallel tasks)
 			if (bIsSearchValid && NextTask->ShouldIgnoreRestartSelf())
 			{
-				for (int32 TestInstanceIdx = 0; TestInstanceIdx <= ActiveInstanceIdx; TestInstanceIdx++)
+				const bool bIsTaskRunning = InstanceStack[ActiveInstanceIdx].HasActiveNode(NextTaskIdx.ExecutionIndex);
+				if (bIsTaskRunning)
 				{
-					const bool bIsTaskRunning = InstanceStack[TestInstanceIdx].HasActiveNode(NextTaskIdx.ExecutionIndex);
-					if (bIsTaskRunning)
-					{
-						BT_SEARCHLOG(SearchData, Verbose, TEXT("Task doesn't allow restart and it's already running! Discaring search."));
-						bIsSearchValid = false;
-						break;
-					}
+					BT_SEARCHLOG(SearchData, Verbose, TEXT("Task doesn't allow restart and it's already running! Discaring search."));
+					bIsSearchValid = false;
 				}
 			}
 		}
diff --git a/Engine/Source/Runtime/AIModule/Private/BehaviorTree/BehaviorTreeTypes.cpp b/Engine/Source/Runtime/AIModule/Private/BehaviorTree/BehaviorTreeTypes.cpp
--- a/Engine/Source/Runtime/AIModule/Private/BehaviorTree/BehaviorTreeTypes.cpp
+++ b/Engine/Source/Runtime/AIModule/Private/BehaviorTree/BehaviorTreeTypes.cpp
@@ -133,7 +133,7 @@ bool FBehaviorTreeInstance::HasActiveNode(uint16 TestExecutionIndex) const
 {
 	if (ActiveNode && ActiveNode->GetExecutionIndex() == TestExecutionIndex)
 	{
-		return true;
+		return (ActiveNodeType == EBTActiveNode::ActiveTask);
 	}
 
 	for (int32 Idx = 0; Idx < ParallelTasks.Num(); Idx++)
diff --git a/Engine/Source/Runtime/AIModule/Private/BehaviorTree/BlackboardComponent.cpp b/Engine/Source/Runtime/AIModule/Private/BehaviorTree/BlackboardComponent.cpp
--- a/Engine/Source/Runtime/AIModule/Private/BehaviorTree/BlackboardComponent.cpp
+++ b/Engine/Source/Runtime/AIModule/Private/BehaviorTree/BlackboardComponent.cpp
@@ -337,6 +337,27 @@ void UBlackboardComponent::UnregisterObserversFrom(UObject* NotifyOwner)
 	}
 }
 
+void UBlackboardComponent::PauseObserverNotifications()
+{
+	bPausedNotifies = true;
+}
+
+
+void UBlackboardComponent::ResumeObserverNotifications(bool bSendQueuedObserverNotifications)
+{
+	bPausedNotifies = false;
+
+	if (bSendQueuedObserverNotifications)
+	{
+		for (int32 UpdateIndex = 0; UpdateIndex < QueuedUpdates.Num(); UpdateIndex++)
+		{
+			NotifyObservers(QueuedUpdates[UpdateIndex]);
+		}
+	}
+
+	QueuedUpdates.Empty();
+}
+
 void UBlackboardComponent::PauseUpdates()
 {
 	bPausedNotifies = true;
diff --git a/Engine/Source/Runtime/AIModule/Private/GameplayDebugger/GameplayDebuggerCategory_AI.cpp b/Engine/Source/Runtime/AIModule/Private/GameplayDebugger/GameplayDebuggerCategory_AI.cpp
--- a/Engine/Source/Runtime/AIModule/Private/GameplayDebugger/GameplayDebuggerCategory_AI.cpp
+++ b/Engine/Source/Runtime/AIModule/Private/GameplayDebugger/GameplayDebuggerCategory_AI.cpp
@@ -272,32 +272,19 @@ void FGameplayDebuggerCategory_AI::OnDataPackReplicated(int32 DataPackId)
 
 void FGameplayDebuggerCategory_AI::DrawData(APlayerController* OwnerPC, FGameplayDebuggerCanvasContext& CanvasContext)
 {
-	const bool bReducedMode = IsSimulateInEditor();
-	bShowCategoryName = !bReducedMode;
-
 	UWorld* MyWorld = OwnerPC->GetWorld();
 	AActor* SelectedActor = FindLocalDebugActor();
 
+	const bool bReducedMode = IsSimulateInEditor();
+	bShowCategoryName = !bReducedMode || DataPack.bHasController;
+
 	DrawPawnIcons(MyWorld, SelectedActor, OwnerPC ? OwnerPC->GetPawn() : nullptr, CanvasContext);
-	if (bReducedMode)
-	{
-		if (DataPack.bHasController)
-		{
-			DrawPath(MyWorld);
-		}
-	}
-	else
+	if (SelectedActor)
 	{
-		if (SelectedActor)
-		{
-			DrawOverheadInfo(*SelectedActor, CanvasContext);
-		}
-
-		DrawPath(MyWorld);
+		DrawOverheadInfo(*SelectedActor, CanvasContext);
 	}
 
-	const bool bShowClassNames = !bReducedMode || DataPack.bHasController;
-	if (bShowClassNames)
+	if (DataPack.bHasController)
 	{
 		CanvasContext.Printf(TEXT("Controller Name: {yellow}%s"), *DataPack.ControllerName);
 		CanvasContext.Printf(TEXT("Pawn Name: {yellow}%s"), *DataPack.PawnName);
diff --git a/Engine/Source/Runtime/AnimGraphRuntime/Private/AnimNodes/AnimNode_LayeredBoneBlend.cpp b/Engine/Source/Runtime/AnimGraphRuntime/Private/AnimNodes/AnimNode_LayeredBoneBlend.cpp
--- a/Engine/Source/Runtime/AnimGraphRuntime/Private/AnimNodes/AnimNode_LayeredBoneBlend.cpp
+++ b/Engine/Source/Runtime/AnimGraphRuntime/Private/AnimNodes/AnimNode_LayeredBoneBlend.cpp
@@ -6,6 +6,7 @@
 #include "Animation/AnimStats.h"
 #include "Animation/AnimInstanceProxy.h"
 
+#define DEFAULT_SOURCEINDEX 0xFF
 /////////////////////////////////////////////////////
 // FAnimNode_LayeredBoneBlend
 
@@ -36,11 +37,38 @@ void FAnimNode_LayeredBoneBlend::ReinitializeBoneBlendWeights(const FBoneContain
 {
 	FAnimationRuntime::CreateMaskWeights(DesiredBoneBlendWeights, LayerSetup, RequiredBones, Skeleton);
 
-	CurrentBoneBlendWeights.Empty(DesiredBoneBlendWeights.Num());
+	CurrentBoneBlendWeights.Reset(DesiredBoneBlendWeights.Num());
 	CurrentBoneBlendWeights.AddZeroed(DesiredBoneBlendWeights.Num());
 
 	//Reinitialize bone blend weights now that we have cleared them
 	FAnimationRuntime::UpdateDesiredBoneWeight(DesiredBoneBlendWeights, CurrentBoneBlendWeights, BlendWeights);
+
+	TArray<SmartName::UID_Type> const & CurveUIDs = RequiredBones.GetAnimCurveNameUids();
+	CurvePoseSourceIndices.Reset(CurveUIDs.Num());
+	// initialize with FF - which is default
+	CurvePoseSourceIndices.Init(DEFAULT_SOURCEINDEX, CurveUIDs.Num());
+
+	// now go through point to correct source indices. Curve only picks one source index
+	for (int32 CurveBlendIter = 0; CurveBlendIter < CurvePoseSourceIndices.Num(); ++CurveBlendIter)
+	{
+		SmartName::UID_Type CurveUID = CurveUIDs[CurveBlendIter];
+		
+		const FCurveMetaData* CurveMetaData = Skeleton->GetCurveMetaData(CurveUID);
+		if (CurveMetaData && CurveMetaData->LinkedBones.Num() > 0)
+		{
+			for (int32 LinkedBoneIndex = 0; LinkedBoneIndex < CurveMetaData->LinkedBones.Num(); ++LinkedBoneIndex)
+			{
+				FCompactPoseBoneIndex CompactPoseIndex = CurveMetaData->LinkedBones[LinkedBoneIndex].GetCompactPoseIndex(RequiredBones);
+				if (CompactPoseIndex != INDEX_NONE)
+				{
+					if (DesiredBoneBlendWeights[CompactPoseIndex.GetInt()].BlendWeight > 0.f)
+					{
+						CurvePoseSourceIndices[CurveBlendIter] = DesiredBoneBlendWeights[CompactPoseIndex.GetInt()].SourceIndex;
+					}
+				}
+			}
+		}
+	}
 }
 
 void FAnimNode_LayeredBoneBlend::CacheBones(const FAnimationCacheBonesContext& Context) 
@@ -141,6 +169,26 @@ void FAnimNode_LayeredBoneBlend::Evaluate(FPoseContext& Output)
 			}
 		}
 
+		// filter to make sure it only includes curves that is linked to the correct bone filter
+		const TArray<SmartName::UID_Type>& UIDList = (*Output.Curve.UIDList);
+		for (int32 CurveIndex = 0; CurveIndex < CurvePoseSourceIndices.Num(); ++CurveIndex)
+		{
+			int32 SourceIndex = CurvePoseSourceIndices[CurveIndex];
+			if (SourceIndex != DEFAULT_SOURCEINDEX)
+			{
+				// if source index is set, clear base pose curve value
+				BasePoseContext.Curve.Set(UIDList[CurveIndex], 0.f);
+				for (int32 ChildIndex = 0; ChildIndex < NumPoses; ++ChildIndex)
+				{
+					if (SourceIndex != ChildIndex)
+					{
+						// if not source, clear it
+						TargetBlendCurves[ChildIndex].Set(UIDList[CurveIndex], 0.f);
+					}
+				}
+			}
+		}
+
 		FAnimationRuntime::BlendPosesPerBoneFilter(BasePoseContext.Pose, TargetBlendPoses, BasePoseContext.Curve, TargetBlendCurves, Output.Pose, Output.Curve, CurrentBoneBlendWeights, bMeshSpaceRotationBlend, CurveBlendOption);
 	}
 }
diff --git a/Engine/Source/Runtime/AnimGraphRuntime/Private/AnimNodes/AnimNode_ModifyCurve.cpp b/Engine/Source/Runtime/AnimGraphRuntime/Private/AnimNodes/AnimNode_ModifyCurve.cpp
new file mode 100644
--- /dev/null
+++ b/Engine/Source/Runtime/AnimGraphRuntime/Private/AnimNodes/AnimNode_ModifyCurve.cpp
@@ -0,0 +1,73 @@
+// Copyright 1998-2016 Epic Games, Inc. All Rights Reserved.
+
+#include "AnimGraphRuntimePrivatePCH.h"
+#include "AnimationRuntime.h"
+#include "AnimInstanceProxy.h"
+#include "AnimNodes/AnimNode_ModifyCurve.h"
+
+FAnimNode_ModifyCurve::FAnimNode_ModifyCurve()
+{
+	ApplyMode = EModifyCurveApplyMode::Blend;
+	Alpha = 1.f;
+}
+
+void FAnimNode_ModifyCurve::Evaluate(FPoseContext& Output)
+{
+	FPoseContext SourceData(Output);
+	SourcePose.Evaluate(SourceData);
+
+	Output = SourceData;
+
+	//	Morph target and Material parameter curves
+	USkeleton* Skeleton = Output.AnimInstanceProxy->GetSkeleton();
+
+	check(CurveNames.Num() == CurveValues.Num());
+
+	for (int32 ModIdx = 0; ModIdx < CurveNames.Num(); ModIdx++)
+	{
+		FName CurveName = CurveNames[ModIdx];
+		SmartName::UID_Type NameUID = Skeleton->GetUIDByName(USkeleton::AnimCurveMappingName, CurveName);
+		if (NameUID != SmartName::MaxUID)
+		{
+			float CurveValue = CurveValues[ModIdx];
+			float CurrentValue = Output.Curve.Get(NameUID);
+
+			// Use ApplyMode enum to decide how to apply
+			if (ApplyMode == EModifyCurveApplyMode::Add)
+			{
+				Output.Curve.Set(NameUID, CurrentValue + CurveValue);
+			}
+			else if (ApplyMode == EModifyCurveApplyMode::Scale)
+			{
+				Output.Curve.Set(NameUID, CurrentValue * CurveValue);
+			}
+			else // Blend
+			{
+				float UseAlpha = FMath::Clamp(Alpha, 0.f, 1.f);
+				Output.Curve.Set(NameUID, FMath::Lerp(CurrentValue, CurveValue, UseAlpha));
+			}
+		}
+	}
+}
+
+void FAnimNode_ModifyCurve::Update(const FAnimationUpdateContext& Context)
+{
+	// Evaluate any BP logic plugged into this node
+	EvaluateGraphExposedInputs.Execute(Context);
+}
+
+#if WITH_EDITOR
+
+void FAnimNode_ModifyCurve::AddCurve(const FName& InName, float InValue)
+{
+	CurveValues.Add(InValue);
+	CurveNames.Add(InName);
+}
+
+void FAnimNode_ModifyCurve::RemoveCurve(int32 PoseIndex)
+{
+	CurveValues.RemoveAt(PoseIndex);
+	CurveNames.RemoveAt(PoseIndex);
+}
+
+#endif // WITH_EDITOR
diff --git a/Engine/Source/Runtime/AnimGraphRuntime/Private/AnimNodes/AnimNode_PoseSnapshot.cpp b/Engine/Source/Runtime/AnimGraphRuntime/Private/AnimNodes/AnimNode_PoseSnapshot.cpp
new file mode 100644
--- /dev/null
+++ b/Engine/Source/Runtime/AnimGraphRuntime/Private/AnimNodes/AnimNode_PoseSnapshot.cpp
@@ -0,0 +1,38 @@
+// Copyright 1998-2016 Epic Games, Inc. All Rights Reserved.
+
+#include "AnimGraphRuntimePrivatePCH.h"
+#include "AnimNodes/AnimNode_PoseSnapshot.h"
+#include "AnimInstanceProxy.h"
+
+/////////////////////////////////////////////////////
+// FAnimCachePoseNode
+
+
+void FAnimNode_PoseSnapshot::Evaluate(FPoseContext& Output)
+{
+	FCompactPose& OutPose = Output.Pose;
+	OutPose.ResetToRefPose();
+
+	if(const FAnimInstanceProxy::FPoseSnapshot* PoseSnapshot = Output.AnimInstanceProxy->GetPoseSnapshot(SnapshotName))
+	{
+		const TArray<FTransform>& LocalTMs = PoseSnapshot->LocalTransforms;
+		const FBoneContainer& RequiredBones = OutPose.GetBoneContainer();
+		for (FCompactPoseBoneIndex PoseBoneIndex : OutPose.ForEachBoneIndex())
+		{
+			const int32 SkeletonBoneIndex = RequiredBones.GetSkeletonIndex(PoseBoneIndex);
+			const int32 MeshBoneIndex = RequiredBones.GetSkeletonToPoseBoneIndexArray()[SkeletonBoneIndex];
+
+			if (LocalTMs.IsValidIndex(MeshBoneIndex))
+			{
+				OutPose[PoseBoneIndex] = LocalTMs[MeshBoneIndex];
+			}
+		}
+	}
+}
+
+
+void FAnimNode_PoseSnapshot::GatherDebugData(FNodeDebugData& DebugData)
+{
+	FString DebugLine = DebugData.GetNodeName(this) + " Snapshot Name:" + SnapshotName.ToString();
+	DebugData.AddDebugItem(DebugLine, true);
+}
\ No newline at end of file
diff --git a/Engine/Source/Runtime/AnimGraphRuntime/Private/BoneControllers/AnimNode_CopyBoneDelta.cpp b/Engine/Source/Runtime/AnimGraphRuntime/Private/BoneControllers/AnimNode_CopyBoneDelta.cpp
--- a/Engine/Source/Runtime/AnimGraphRuntime/Private/BoneControllers/AnimNode_CopyBoneDelta.cpp
+++ b/Engine/Source/Runtime/AnimGraphRuntime/Private/BoneControllers/AnimNode_CopyBoneDelta.cpp
@@ -44,7 +44,7 @@ void FAnimNode_CopyBoneDelta::EvaluateBoneTransforms(USkeletalMeshComponent* Ske
 	FAnimationRuntime::ConvertCSTransformToBoneSpace(SkelComp, MeshBases, TargetTM, TargetBoneIndex, BCS_ParentBoneSpace);
 
 	// Ref pose transform
-	FTransform RefLSTransform = SkelComp->SkeletalMesh->RefSkeleton.GetRefBonePose()[SourceBone.GetMeshPoseIndex().GetInt()];
+	FTransform RefLSTransform = SkelComp->SkeletalMesh->RefSkeleton.GetRefBonePose()[SourceBone.GetMeshPoseIndex(BoneContainer).GetInt()];
 
 	// Get transform relative to ref pose
 	SourceTM.SetToRelativeTransform(RefLSTransform);
diff --git a/Engine/Source/Runtime/AnimGraphRuntime/Public/AnimNodes/AnimNode_LayeredBoneBlend.h b/Engine/Source/Runtime/AnimGraphRuntime/Public/AnimNodes/AnimNode_LayeredBoneBlend.h
--- a/Engine/Source/Runtime/AnimGraphRuntime/Public/AnimNodes/AnimNode_LayeredBoneBlend.h
+++ b/Engine/Source/Runtime/AnimGraphRuntime/Public/AnimNodes/AnimNode_LayeredBoneBlend.h
@@ -43,6 +43,7 @@ public:
 protected:
 	TArray<FPerBoneBlendWeight> DesiredBoneBlendWeights;
 	TArray<FPerBoneBlendWeight> CurrentBoneBlendWeights;
+	TArray<uint8> CurvePoseSourceIndices;
 
 public:	
 	FAnimNode_LayeredBoneBlend()
diff --git a/Engine/Source/Runtime/AnimGraphRuntime/Public/AnimNodes/AnimNode_ModifyCurve.h b/Engine/Source/Runtime/AnimGraphRuntime/Public/AnimNodes/AnimNode_ModifyCurve.h
new file mode 100644
--- /dev/null
+++ b/Engine/Source/Runtime/AnimGraphRuntime/Public/AnimNodes/AnimNode_ModifyCurve.h
@@ -0,0 +1,55 @@
+// Copyright 1998-2016 Epic Games, Inc. All Rights Reserved.
+
+#pragma once
+
+#include "Animation/AnimNodeBase.h"
+#include "AnimNode_ModifyCurve.generated.h"
+
+UENUM()
+enum class EModifyCurveApplyMode : uint8
+{
+	/** Add new value to input curve value */
+	Add,
+
+	/** Scale input value by new value */
+	Scale,
+
+	/** Blend input with new curve value, using Alpha setting on the node */
+	Blend
+};
+
+/** Easy way to modify curve values on a pose */
+USTRUCT()
+struct ANIMGRAPHRUNTIME_API FAnimNode_ModifyCurve : public FAnimNode_Base
+{
+	GENERATED_BODY()
+
+	UPROPERTY(EditAnywhere, EditFixedSize, BlueprintReadWrite, Category = Links)
+	FPoseLink SourcePose;
+
+	UPROPERTY(EditAnywhere, BlueprintReadWrite, Category = ModifyCurve)
+	EModifyCurveApplyMode ApplyMode;
+
+	UPROPERTY(EditAnywhere, BlueprintReadWrite, editfixedsize, Category = ModifyCurve, meta = (PinShownByDefault))
+	TArray<float> CurveValues;
+
+	UPROPERTY()
+	TArray<FName> CurveNames;
+
+	UPROPERTY(EditAnywhere, BlueprintReadWrite, Category = ModifyCurve, meta = (PinShownByDefault))
+	float Alpha;
+
+	FAnimNode_ModifyCurve();
+
+	// FAnimNode_Base interface
+	virtual void Evaluate(FPoseContext& Output) override;
+	virtual void Update(const FAnimationUpdateContext& Context) override;
+	// End of FAnimNode_Base interface
+
+#if WITH_EDITOR
+	/** Add new curve being modified */
+	void AddCurve(const FName& InName, float InValue);
+	/** Remove a curve from being modified */
+	void RemoveCurve(int32 PoseIndex);
+#endif // WITH_EDITOR
+};
diff --git a/Engine/Source/Runtime/AnimGraphRuntime/Public/AnimNodes/AnimNode_PoseSnapshot.h b/Engine/Source/Runtime/AnimGraphRuntime/Public/AnimNodes/AnimNode_PoseSnapshot.h
new file mode 100644
--- /dev/null
+++ b/Engine/Source/Runtime/AnimGraphRuntime/Public/AnimNodes/AnimNode_PoseSnapshot.h
@@ -0,0 +1,20 @@
+// Copyright 1998-2016 Epic Games, Inc. All Rights Reserved.
+
+#pragma once
+#include "Animation/AnimNodeBase.h"
+#include "AnimNode_PoseSnapshot.generated.h"
+
+// RefPose pose nodes - ref pose or additive RefPose pose
+USTRUCT()
+struct ANIMGRAPHRUNTIME_API FAnimNode_PoseSnapshot : public FAnimNode_Base
+{
+	GENERATED_USTRUCT_BODY()
+
+public:	
+
+	virtual void Evaluate(FPoseContext& Output) override;
+	virtual void GatherDebugData(FNodeDebugData& DebugData) override;
+
+	UPROPERTY(EditAnywhere, BlueprintReadWrite, Category = Copy, meta = (PinShownByDefault))
+	FName SnapshotName;
+};
\ No newline at end of file
diff --git a/Engine/Source/Runtime/AssetRegistry/Private/AssetDataGatherer.cpp b/Engine/Source/Runtime/AssetRegistry/Private/AssetDataGatherer.cpp
--- a/Engine/Source/Runtime/AssetRegistry/Private/AssetDataGatherer.cpp
+++ b/Engine/Source/Runtime/AssetRegistry/Private/AssetDataGatherer.cpp
@@ -4,7 +4,7 @@
 
 namespace AssetDataGathererConstants
 {
-	static const int32 CacheSerializationVersion = 7;
+	static const int32 CacheSerializationVersion = 8;
 	static const int32 MaxFilesToDiscoverBeforeFlush = 2500;
 	static const int32 MaxFilesToGatherBeforeFlush = 250;
 	static const int32 MaxFilesToProcessBeforeCacheWrite = 50000;
diff --git a/Engine/Source/Runtime/AssetRegistry/Private/AssetRegistry.cpp b/Engine/Source/Runtime/AssetRegistry/Private/AssetRegistry.cpp
--- a/Engine/Source/Runtime/AssetRegistry/Private/AssetRegistry.cpp
+++ b/Engine/Source/Runtime/AssetRegistry/Private/AssetRegistry.cpp
@@ -11,8 +11,11 @@
 
 DEFINE_LOG_CATEGORY(LogAssetRegistry);
 
-/** Enum for tracking versions of the "mini" runtime asset registry that gets used by iterative cooking and the runtime
-	dependency preloading system */
+/** 
+ * Enum for tracking versions of the "mini" runtime asset registry that gets used by iterative cooking and the runtime
+ * dependency preloading system 
+ * This enum is NOT used for the editor caching, that uses AssetDataGathererConstants::CacheSerializationVersion
+ */
 enum class ERuntimeRegistryVersion
 {
 	PreVersioning,					// From before file versioning was implemented
@@ -288,7 +291,7 @@ FAssetRegistry::~FAssetRegistry()
 	else
 	{
 		// Delete all depends nodes in the cache
-		for (TMap<FName, FDependsNode*>::TConstIterator DependsIt(CachedDependsNodes); DependsIt; ++DependsIt)
+		for (TMap<FAssetIdentifier, FDependsNode*>::TConstIterator DependsIt(CachedDependsNodes); DependsIt; ++DependsIt)
 		{
 			if (DependsIt.Value())
 			{
@@ -879,9 +882,9 @@ bool FAssetRegistry::GetAllAssets(TArray<FAssetData>& OutAssetData, bool bInclud
 	return true;
 }
 
-bool FAssetRegistry::GetDependencies(FName PackageName, TArray<FName>& OutDependencies, EAssetRegistryDependencyType::Type InDependencyType, bool bResolveIniStringReferences) const
+bool FAssetRegistry::GetDependencies(const FAssetIdentifier& AssetIdentifier, TArray<FAssetIdentifier>& OutDependencies, EAssetRegistryDependencyType::Type InDependencyType) const
 {
-	const FDependsNode*const* NodePtr = CachedDependsNodes.Find(PackageName);
+	const FDependsNode*const* NodePtr = CachedDependsNodes.Find(AssetIdentifier);
 	const FDependsNode* Node = nullptr;
 	if (NodePtr != nullptr )
 	{
@@ -890,26 +893,7 @@ bool FAssetRegistry::GetDependencies(FName PackageName, TArray<FName>& OutDepend
 
 	if (Node != nullptr)
 	{
-		if (bResolveIniStringReferences)
-		{
-			TArray<FName> DependencyNodes;
-			Node->GetDependencies(DependencyNodes, InDependencyType);
-
-			for (auto NodeIt = DependencyNodes.CreateConstIterator(); NodeIt; ++NodeIt)
-			{
-				FString PackagePath = NodeIt->ToString();
-				const FString* IniFilename = GetIniFilenameFromObjectsReference(PackagePath);
-
-				OutDependencies.Add(
-					(IniFilename != nullptr)
-					? FName(*ResolveIniObjectsReference(PackagePath, IniFilename))
-					: *NodeIt);
-			}
-		}
-		else
-		{
-			Node->GetDependencies(OutDependencies, InDependencyType);
-		}
+		Node->GetDependencies(OutDependencies, InDependencyType);
 
 		return true;
 	}
@@ -919,9 +903,25 @@ bool FAssetRegistry::GetDependencies(FName PackageName, TArray<FName>& OutDepend
 	}
 }
 
-bool FAssetRegistry::GetReferencers(FName PackageName, TArray<FName>& OutReferencers, EAssetRegistryDependencyType::Type InReferenceType) const
+bool FAssetRegistry::GetDependencies(FName PackageName, TArray<FName>& OutDependencies, EAssetRegistryDependencyType::Type InDependencyType) const
 {
-	const FDependsNode*const* NodePtr = CachedDependsNodes.Find(PackageName);
+	TArray<FAssetIdentifier> TempDependencies;
+
+	if (GetDependencies(FAssetIdentifier(PackageName), TempDependencies, InDependencyType))
+	{
+		for (const FAssetIdentifier& AssetId : TempDependencies)
+		{
+			OutDependencies.AddUnique(AssetId.PackageName);
+		}
+		return true;
+	}
+
+	return false;
+}
+
+bool FAssetRegistry::GetReferencers(const FAssetIdentifier& AssetIdentifier, TArray<FAssetIdentifier>& OutReferencers, EAssetRegistryDependencyType::Type InReferenceType) const
+{
+	const FDependsNode*const* NodePtr = CachedDependsNodes.Find(AssetIdentifier);
 	const FDependsNode* Node = nullptr;
 	bool bShowingAllReferences = InReferenceType == EAssetRegistryDependencyType::All;
 
@@ -946,14 +946,14 @@ bool FAssetRegistry::GetReferencers(FName PackageName, TArray<FName>& OutReferen
 				{
 					if (Dependency == Node)
 					{
-						OutReferencers.Add((*NodeIt)->GetPackageName());
+						OutReferencers.Add((*NodeIt)->GetIdentifier());
 						break;
 					}
 				}
 			}
 			else
 			{
-				OutReferencers.Add((*NodeIt)->GetPackageName());
+				OutReferencers.Add((*NodeIt)->GetIdentifier());
 			}
 		}
 
@@ -965,6 +965,22 @@ bool FAssetRegistry::GetReferencers(FName PackageName, TArray<FName>& OutReferen
 	}
 }
 
+bool FAssetRegistry::GetReferencers(FName PackageName, TArray<FName>& OutReferencers, EAssetRegistryDependencyType::Type InReferenceType) const
+{
+	TArray<FAssetIdentifier> TempReferencers;
+
+	if (GetReferencers(FAssetIdentifier(PackageName), TempReferencers, InReferenceType))
+	{
+		for (const FAssetIdentifier& AssetId : TempReferencers)
+		{
+			OutReferencers.AddUnique(AssetId.PackageName);
+		}
+		return true;
+	}
+
+	return false;
+}
+
 bool FAssetRegistry::GetAncestorClassNames(FName ClassName, TArray<FName>& OutAncestorClassNames) const
 {
 	// Start with the cached inheritance map
@@ -1456,6 +1472,33 @@ bool FAssetRegistry::IsLoadingAssets() const
 	return !bInitialSearchCompleted;
 }
 
+IAssetRegistry::FAssetEditSearchableNameDelegate& FAssetRegistry::OnEditSearchableName(FName PackageName, FName ObjectName)
+{
+	return EditSearchableNameDelegates.FindOrAdd(FAssetIdentifier(PackageName, ObjectName));
+}
+
+bool FAssetRegistry::EditSearchableName(const FAssetIdentifier& SearchableName)
+{
+	for (TPair<FAssetIdentifier, FAssetEditSearchableNameDelegate> MapPair : EditSearchableNameDelegates)
+	{
+		if (MapPair.Key.PackageName == SearchableName.PackageName &&
+			(MapPair.Key.ObjectName == SearchableName.ObjectName || MapPair.Key.ObjectName.IsNone()) &&
+			(MapPair.Key.ValueName == SearchableName.ValueName || MapPair.Key.ValueName.IsNone()))
+		{
+			// Try this callback
+			if (MapPair.Value.IsBound())
+			{
+				if (MapPair.Value.Execute(SearchableName))
+				{
+					return true;
+				}
+			}
+		}
+	}
+
+	return false;
+}
+
 void FAssetRegistry::Tick(float DeltaTime)
 {
 	double TickStartTime = FPlatformTime::Seconds();
@@ -1549,6 +1592,7 @@ bool FAssetRegistry::IsUsingWorldAssets()
 
 void FAssetRegistry::Serialize(FArchive& Ar)
 {
+	// This is only used for the runtime version of the AssetRegistry
 	if (Ar.IsSaving())
 	{
 		check(CachedAssetsByObjectPath.Num() == NumAssets);
@@ -2167,7 +2211,9 @@ void FAssetRegistry::DependencyDataGathered(const double TickStartTime, TArray<F
 
 		for (const FString& StringAssetReference : Result.StringAssetReferencesMap)
 		{
-			const FName AssetReference = *StringAssetReference;
+			// Possibly resolve ini:name references before adding to dependency list
+			const FString* IniFilename = GetIniFilenameFromObjectsReference(StringAssetReference);
+			const FName AssetReference = IniFilename ? *ResolveIniObjectsReference(StringAssetReference, IniFilename) : *StringAssetReference;
 
 			// Already processed?
 			if (PackageDependencies.Contains(AssetReference))
@@ -2175,7 +2221,7 @@ void FAssetRegistry::DependencyDataGathered(const double TickStartTime, TArray<F
 				continue;
 			}
 
-			if (FPackageName::IsShortPackageName(StringAssetReference))
+			if (FPackageName::IsShortPackageName(AssetReference))
 			{
 				UE_LOG(LogAssetRegistry, Warning, TEXT("Package with string asset reference with short asset path: %s. This is unsupported, can couse errors and be slow on loading. Please resave the package to fix this."), *Result.PackageName.ToString());
 			}
@@ -2183,6 +2229,41 @@ void FAssetRegistry::DependencyDataGathered(const double TickStartTime, TArray<F
 			PackageDependencies.Add(AssetReference, EAssetRegistryDependencyType::Soft);
 		}
 
+		for (const TPair<FPackageIndex, TArray<FName>>& SearchableNameList : Result.SearchableNamesMap)
+		{
+			FName ObjectName;
+			FName PackageName;
+
+			// Find object and package name from linker
+			for (FPackageIndex LinkerIndex = SearchableNameList.Key; !LinkerIndex.IsNull();)
+			{
+				FObjectResource& Resource = Result.ImpExp(LinkerIndex);
+				LinkerIndex = Resource.OuterIndex;
+				if (ObjectName.IsNone() && !LinkerIndex.IsNull())
+				{
+					ObjectName = Resource.ObjectName;
+				}
+				else if (LinkerIndex.IsNull())
+				{
+					PackageName = Resource.ObjectName;
+				}
+			}
+
+			for (FName NameReference : SearchableNameList.Value)
+			{
+				FAssetIdentifier AssetId = FAssetIdentifier(PackageName, ObjectName, NameReference);
+
+				// Add node for all name references
+				FDependsNode* DependsNode = CreateOrFindDependsNode(AssetId);
+
+				if (DependsNode != nullptr)
+				{
+					Node->AddDependency(DependsNode, EAssetRegistryDependencyType::SearchableName);
+					DependsNode->AddReferencer(Node);
+				}
+			}	
+		}
+
 		// Doubly-link all new dependencies for this package
 		for (auto NewDependsIt : PackageDependencies)
 		{
@@ -2236,9 +2317,9 @@ void FAssetRegistry::CookedPackageNamesWithoutAssetDataGathered(const double Tic
 	}
 }
 
-FDependsNode* FAssetRegistry::FindDependsNode(FName ObjectName)
+FDependsNode* FAssetRegistry::FindDependsNode(const FAssetIdentifier& Identifier)
 {
-	FDependsNode** FoundNode = CachedDependsNodes.Find(ObjectName);
+	FDependsNode** FoundNode = CachedDependsNodes.Find(Identifier);
 	if (FoundNode)
 	{
 		return *FoundNode;
@@ -2249,24 +2330,24 @@ FDependsNode* FAssetRegistry::FindDependsNode(FName ObjectName)
 	}
 }
 
-FDependsNode* FAssetRegistry::CreateOrFindDependsNode(FName ObjectPath)
+FDependsNode* FAssetRegistry::CreateOrFindDependsNode(const FAssetIdentifier& Identifier)
 {
-	FDependsNode* FoundNode = FindDependsNode(ObjectPath);
+	FDependsNode* FoundNode = FindDependsNode(Identifier);
 	if ( FoundNode )
 	{
 		return FoundNode;
 	}
 
-	FDependsNode* NewNode = new FDependsNode(ObjectPath);
+	FDependsNode* NewNode = new FDependsNode(Identifier);
 	NumDependsNodes++;
-	CachedDependsNodes.Add(ObjectPath, NewNode);
+	CachedDependsNodes.Add(Identifier, NewNode);
 
 	return NewNode;
 }
 
-bool FAssetRegistry::RemoveDependsNode( FName PackageName )
+bool FAssetRegistry::RemoveDependsNode(const FAssetIdentifier& Identifier)
 {
-	FDependsNode** NodePtr = CachedDependsNodes.Find( PackageName );
+	FDependsNode** NodePtr = CachedDependsNodes.Find( Identifier );
 
 	if ( NodePtr != nullptr )
 	{
@@ -2292,7 +2373,7 @@ bool FAssetRegistry::RemoveDependsNode( FName PackageName )
 			}
 
 			// Remove the node and delete it
-			CachedDependsNodes.Remove( PackageName );
+			CachedDependsNodes.Remove( Identifier );
 			NumDependsNodes--;
 			
 			// if the depends nodes were preallocated in a block, we can't delete them one at a time, only the whole chunk in the destructor
diff --git a/Engine/Source/Runtime/AssetRegistry/Private/AssetRegistry.h b/Engine/Source/Runtime/AssetRegistry/Private/AssetRegistry.h
--- a/Engine/Source/Runtime/AssetRegistry/Private/AssetRegistry.h
+++ b/Engine/Source/Runtime/AssetRegistry/Private/AssetRegistry.h
@@ -21,8 +21,10 @@ public:
 	virtual bool GetAssets(const FARFilter& Filter, TArray<FAssetData>& OutAssetData) const override;
 	virtual FAssetData GetAssetByObjectPath( const FName ObjectPath, bool bIncludeOnlyOnDiskAssets = false ) const override;
 	virtual bool GetAllAssets(TArray<FAssetData>& OutAssetData, bool bIncludeOnlyOnDiskAssets = false) const override;
-	virtual bool GetDependencies(FName PackageName, TArray<FName>& OutDependencies, EAssetRegistryDependencyType::Type InDependencyType = EAssetRegistryDependencyType::All, bool bResolveIniStringReferences = false) const override;
-	virtual bool GetReferencers(FName PackageName, TArray<FName>& OutReferencers, EAssetRegistryDependencyType::Type InReferenceType) const override;
+	virtual bool GetDependencies(const FAssetIdentifier& AssetIdentifier, TArray<FAssetIdentifier>& OutDependencies, EAssetRegistryDependencyType::Type InDependencyType = EAssetRegistryDependencyType::All) const override;
+	virtual bool GetDependencies(FName PackageName, TArray<FName>& OutDependencies, EAssetRegistryDependencyType::Type InDependencyType = EAssetRegistryDependencyType::Packages) const override;
+	virtual bool GetReferencers(const FAssetIdentifier& AssetIdentifier, TArray<FAssetIdentifier>& OutReferencers, EAssetRegistryDependencyType::Type InReferenceType = EAssetRegistryDependencyType::All) const override;
+	virtual bool GetReferencers(FName PackageName, TArray<FName>& OutReferencers, EAssetRegistryDependencyType::Type InReferenceType = EAssetRegistryDependencyType::Packages) const override;
 	virtual bool GetAncestorClassNames(FName ClassName, TArray<FName>& OutAncestorClassNames) const override;
 	virtual void GetDerivedClassNames(const TArray<FName>& ClassNames, const TSet<FName>& ExcludedClassNames, TSet<FName>& OutDerivedClassNames) const override;
 	virtual void GetAllCachedPaths(TArray<FString>& OutPathList) const override;
@@ -75,6 +77,9 @@ public:
 	DECLARE_DERIVED_EVENT( FAssetRegistry, IAssetRegistry::FFileLoadProgressUpdatedEvent, FFileLoadProgressUpdatedEvent );
 	virtual FFileLoadProgressUpdatedEvent& OnFileLoadProgressUpdated() override { return FileLoadProgressUpdatedEvent; }
 
+	virtual IAssetRegistry::FAssetEditSearchableNameDelegate& OnEditSearchableName(FName PackageName, FName ObjectName) override;
+	virtual bool EditSearchableName(const FAssetIdentifier& SearchableName) override;
+
 	virtual bool IsLoadingAssets() const override;
 
 	virtual void Tick (float DeltaTime) override;
@@ -101,13 +106,13 @@ private:
 	void CookedPackageNamesWithoutAssetDataGathered(const double TickStartTime, TArray<FString>& CookedPackageNamesWithoutAssetDataResults);
 
 	/** Finds an existing node for the given package and returns it, or returns null if one isn't found */
-	FDependsNode* FindDependsNode(FName ObjectName);
+	FDependsNode* FindDependsNode(const FAssetIdentifier& Identifier);
 
 	/** Creates a node in the CachedDependsNodes map or finds the existing node and returns it */
-	FDependsNode* CreateOrFindDependsNode(FName ObjectName);
+	FDependsNode* CreateOrFindDependsNode(const FAssetIdentifier& Identifier);
 
 	/** Removes the depends node and updates the dependencies to no longer contain it as as a referencer. */
-	bool RemoveDependsNode( FName PackageName );
+	bool RemoveDependsNode(const FAssetIdentifier& Identifier);
 
 	/** Adds an asset to the empty package list which contains packages that have no assets left in them */
 	void AddEmptyPackage(FName PackageName);
@@ -202,7 +207,7 @@ private:
 	TMap<FName, TArray<FAssetData*> > CachedAssetsByTag;
 
 	/** A map of object names to dependency data */
-	TMap<FName, FDependsNode*> CachedDependsNodes;
+	TMap<FAssetIdentifier, FDependsNode*> CachedDependsNodes;
 
 	/** The set of empty package names (packages which contain no assets but have not yet been saved) */
 	TSet<FName> CachedEmptyPackages;
@@ -258,6 +263,9 @@ private:
 	/** The delegate to execute while loading files to update progress */
 	FFileLoadProgressUpdatedEvent FileLoadProgressUpdatedEvent;
 
+	/** Delegates to call when editing searchable name */
+	TMap<FAssetIdentifier, FAssetEditSearchableNameDelegate> EditSearchableNameDelegates;
+
 	/** Counters for asset/depends data memory allocation to ensure that every FAssetData and FDependsNode created is deleted */
 	int32 NumAssets;
 	int32 NumDependsNodes;
diff --git a/Engine/Source/Runtime/AssetRegistry/Private/DependsNode.cpp b/Engine/Source/Runtime/AssetRegistry/Private/DependsNode.cpp
--- a/Engine/Source/Runtime/AssetRegistry/Private/DependsNode.cpp
+++ b/Engine/Source/Runtime/AssetRegistry/Private/DependsNode.cpp
@@ -2,14 +2,9 @@
 
 #include "AssetRegistryPCH.h"
 
-FDependsNode::FDependsNode(FName InPackageName)
-{
-	PackageName = InPackageName;
-}
-
 void FDependsNode::PrintNode() const
 {
-	UE_LOG(LogAssetRegistry, Log, TEXT("*** Printing DependsNode: %s ***"), *PackageName.ToString());
+	UE_LOG(LogAssetRegistry, Log, TEXT("*** Printing DependsNode: %s ***"), *Identifier.ToString());
 	UE_LOG(LogAssetRegistry, Log, TEXT("*** Dependencies:"));
 	PrintDependencies();
 	UE_LOG(LogAssetRegistry, Log, TEXT("*** Referencers:"));
@@ -39,11 +34,11 @@ void FDependsNode::GetDependencies(TArray<FDependsNode*>& OutDependencies, EAsse
 	InDependencyType);
 }
 
-void FDependsNode::GetDependencies(TArray<FName>& OutDependencies, EAssetRegistryDependencyType::Type InDependencyType) const
+void FDependsNode::GetDependencies(TArray<FAssetIdentifier>& OutDependencies, EAssetRegistryDependencyType::Type InDependencyType) const
 {
 	IterateOverDependencies([&](FDependsNode* InDependency, EAssetRegistryDependencyType::Type /*InDependencyType*/)
 	{
-		OutDependencies.Add(InDependency->GetPackageName());
+		OutDependencies.Add(InDependency->GetIdentifier());
 	},
 	InDependencyType);
 }
@@ -64,11 +59,11 @@ void FDependsNode::PrintDependenciesRecursive(const FString& Indent, TSet<const
 	}
 	else if ( VisitedNodes.Contains(this) )
 	{
-		UE_LOG(LogAssetRegistry, Log, TEXT("%s[CircularReferenceTo]%s"), *Indent, *PackageName.ToString());
+		UE_LOG(LogAssetRegistry, Log, TEXT("%s[CircularReferenceTo]%s"), *Indent, *Identifier.ToString());
 	}
 	else
 	{
-		UE_LOG(LogAssetRegistry, Log, TEXT("%s%s"), *Indent, *PackageName.ToString());
+		UE_LOG(LogAssetRegistry, Log, TEXT("%s%s"), *Indent, *Identifier.ToString());
 		VisitedNodes.Add(this);
 
 		IterateOverDependencies([&](FDependsNode* InDependency, EAssetRegistryDependencyType::Type /*InDependencyType*/)
@@ -88,11 +83,11 @@ void FDependsNode::PrintReferencersRecursive(const FString& Indent, TSet<const F
 	}
 	else if ( VisitedNodes.Contains(this) )
 	{
-		UE_LOG(LogAssetRegistry, Log, TEXT("%s[CircularReferenceTo]%s"), *Indent, *PackageName.ToString());
+		UE_LOG(LogAssetRegistry, Log, TEXT("%s[CircularReferenceTo]%s"), *Indent, *Identifier.ToString());
 	}
 	else
 	{
-		UE_LOG(LogAssetRegistry, Log, TEXT("%s%s"), *Indent, *PackageName.ToString());
+		UE_LOG(LogAssetRegistry, Log, TEXT("%s%s"), *Indent, *Identifier.ToString());
 		VisitedNodes.Add(this);
 
 		for (auto ReferencerIt = Referencers.CreateConstIterator(); ReferencerIt; ++ReferencerIt)
diff --git a/Engine/Source/Runtime/AssetRegistry/Private/DependsNode.h b/Engine/Source/Runtime/AssetRegistry/Private/DependsNode.h
--- a/Engine/Source/Runtime/AssetRegistry/Private/DependsNode.h
+++ b/Engine/Source/Runtime/AssetRegistry/Private/DependsNode.h
@@ -6,8 +6,8 @@
 class FDependsNode
 {
 public:
-	FDependsNode() { PackageName = NAME_None; }
-	FDependsNode(FName InPackageName);
+	FDependsNode() {}
+	FDependsNode(const FAssetIdentifier& InIdentifier) : Identifier(InIdentifier) {}
 
 	/** Prints the dependencies and referencers for this node to the log */
 	void PrintNode() const;
@@ -18,13 +18,17 @@ public:
 	/** Gets the list of dependencies for this node */
 	void GetDependencies(TArray<FDependsNode*>& OutDependencies, EAssetRegistryDependencyType::Type InDependencyType = EAssetRegistryDependencyType::All) const;
 	/** Gets the list of dependency names for this node */
-	void GetDependencies(TArray<FName>& OutDependencies, EAssetRegistryDependencyType::Type InDependencyType = EAssetRegistryDependencyType::All) const;
+	void GetDependencies(TArray<FAssetIdentifier>& OutDependencies, EAssetRegistryDependencyType::Type InDependencyType = EAssetRegistryDependencyType::All) const;
 	/** Gets the list of referencers to this node */
 	void GetReferencers(TArray<FDependsNode*>& OutReferencers) const;
 	/** Gets the name of the package that this node represents */
-	FName GetPackageName() const { return PackageName; }
+	FName GetPackageName() const { return Identifier.PackageName; }
 	/** Sets the name of the package that this node represents */
-	void SetPackageName(FName InName) { PackageName = InName; }
+	void SetPackageName(FName InName) { Identifier = FAssetIdentifier(InName); }
+	/** Returns the entire identifier */
+	FAssetIdentifier GetIdentifier() const { return Identifier; }
+	/** Sets the entire identifier */
+	void SetIdentifier(const FAssetIdentifier& InIdentifier) { Identifier = InIdentifier;  }
 	/** Add a dependency to this node */
 	void AddDependency(FDependsNode* InDependency, EAssetRegistryDependencyType::Type InDependencyType = EAssetRegistryDependencyType::Hard)
 	{ 
@@ -32,56 +36,60 @@ public:
 		{
 			HardDependencies.AddUnique(InDependency);
 		}
-		else
+		if (InDependencyType == EAssetRegistryDependencyType::Soft)
 		{
 			SoftDependencies.AddUnique(InDependency);
 		}
+		if (InDependencyType == EAssetRegistryDependencyType::SearchableName)
+		{
+			NameDependencies.AddUnique(InDependency);
+		}
 	}
 	/** Add a referencer to this node */
 	void AddReferencer(FDependsNode* InReferencer) { Referencers.AddUnique(InReferencer); }
 	/** Remove a dependency from this node */
-	void RemoveDependency(FDependsNode* InDependency) { HardDependencies.Remove(InDependency); SoftDependencies.Remove(InDependency); }
+	void RemoveDependency(FDependsNode* InDependency) { HardDependencies.Remove(InDependency); SoftDependencies.Remove(InDependency); NameDependencies.Remove(InDependency); }
 	/** Remove a referencer from this node */
 	void RemoveReferencer(FDependsNode* InReferencer) { Referencers.Remove(InReferencer); }
 	/** Clear all dependency records from this node */
-	void ClearDependencies() { HardDependencies.Empty(); SoftDependencies.Empty(); }
+	void ClearDependencies() { HardDependencies.Empty(); SoftDependencies.Empty(); NameDependencies.Empty(); }
 
-	/** Iterate over all the dependencies of this node, filtered by the supplied type parameter, and call the supplied lambda
-		parameter on the record */
+	/** Iterate over all the dependencies of this node, filtered by the supplied type parameter, and call the supplied lambda parameter on the record */
 	template <class T>
 	void IterateOverDependencies(T InCallback, EAssetRegistryDependencyType::Type InDependencyType = EAssetRegistryDependencyType::All)
 	{
 		if (InDependencyType & EAssetRegistryDependencyType::Hard)
 		{
-			for (auto Dependency : HardDependencies)
+			for (FDependsNode* Dependency : HardDependencies)
 			{
-				if (Dependency)
-				{
-					InCallback(Dependency, EAssetRegistryDependencyType::Hard);
-				}
+				InCallback(Dependency, EAssetRegistryDependencyType::Hard);
 			}
 		}
 
 		if (InDependencyType & EAssetRegistryDependencyType::Soft)
 		{
-			for (auto Dependency : SoftDependencies)
+			for (FDependsNode* Dependency : SoftDependencies)
 			{
-				if (Dependency)
-				{
-					InCallback(Dependency, EAssetRegistryDependencyType::Soft);
-				}
+				InCallback(Dependency, EAssetRegistryDependencyType::Soft);
+			}
+		}
+
+		if (InDependencyType & EAssetRegistryDependencyType::SearchableName)
+		{
+			for (FDependsNode* Dependency : NameDependencies)
+			{
+				InCallback(Dependency, EAssetRegistryDependencyType::SearchableName);
 			}
 		}
 	}
 
-	/** Iterate over all the dependencies of this node, filtered by the supplied type parameter, and call the supplied lambda
-	parameter on the record */
+	/** Iterate over all the dependencies of this node, filtered by the supplied type parameter, and call the supplied lambda parameter on the record */
 	template <class T>
 	void IterateOverDependencies(T InCallback, EAssetRegistryDependencyType::Type InDependencyType = EAssetRegistryDependencyType::All) const
 	{
 		if (InDependencyType & EAssetRegistryDependencyType::Hard)
 		{
-			for (auto Dependency : HardDependencies)
+			for (FDependsNode* Dependency : HardDependencies)
 			{
 				InCallback(Dependency, EAssetRegistryDependencyType::Hard);
 			}
@@ -89,11 +97,19 @@ public:
 
 		if (InDependencyType & EAssetRegistryDependencyType::Soft)
 		{
-			for (auto Dependency : SoftDependencies)
+			for (FDependsNode* Dependency : SoftDependencies)
 			{
 				InCallback(Dependency, EAssetRegistryDependencyType::Soft);
 			}
 		}
+
+		if (InDependencyType & EAssetRegistryDependencyType::SearchableName)
+		{
+			for (FDependsNode* Dependency : NameDependencies)
+			{
+				InCallback(Dependency, EAssetRegistryDependencyType::SearchableName);
+			}
+		}
 	}
 
 	/** Iterate over all the referencers of this node and call the supplied lambda
@@ -101,7 +117,7 @@ public:
 	template <class T>
 	void IterateOverReferencers(T InCallback)
 	{
-		for (auto Referencer : Referencers)
+		for (FDependsNode* Referencer : Referencers)
 		{
 			InCallback(Referencer);
 		}
@@ -120,12 +136,14 @@ private:
 	/** Recursively prints referencers to the node starting with the specified indent. VisitedNodes should be an empty set at first which is populated recursively. */
 	void PrintReferencersRecursive(const FString& Indent, TSet<const FDependsNode*>& VisitedNodes) const;
 
-	/** The name of the package this node represents */
-	FName PackageName;
+	/** The name of the package/object this node represents */
+	FAssetIdentifier Identifier;
 	/** The list of hard dependencies for this node */
 	TArray<FDependsNode*> HardDependencies;
 	/** The list of soft dependencies for this node */
 	TArray<FDependsNode*> SoftDependencies;
+	/** The list of searchable name dependencies for this node */
+	TArray<FDependsNode*> NameDependencies;
 	/** The list of referencers to this node */
 	TArray<FDependsNode*> Referencers;
 };
\ No newline at end of file
diff --git a/Engine/Source/Runtime/AssetRegistry/Private/PackageDependencyData.h b/Engine/Source/Runtime/AssetRegistry/Private/PackageDependencyData.h
--- a/Engine/Source/Runtime/AssetRegistry/Private/PackageDependencyData.h
+++ b/Engine/Source/Runtime/AssetRegistry/Private/PackageDependencyData.h
@@ -25,6 +25,11 @@ public:
 		Ar << DependencyData.ImportMap;
 		Ar << DependencyData.StringAssetReferencesMap;
 
+		if (Ar.UE4Ver() >= VER_UE4_ADDED_SEARCHABLE_NAMES)
+		{
+			Ar << DependencyData.SearchableNamesMap;
+		}
+
 		return Ar;
 	}
 };
\ No newline at end of file
